Sequencer: Re-architect Curve Editor integration in Sequencer.

All logic is now created FSequencerCurveEditorApp, which manages dependencies.
- CurveEditor/Integration: All the subsystems that integrate Curve Editor and Sequencer.
- CurveEditor/Specializations: Holds all the class specializations & extensions that are injected into FCurveEditor
- FSequencerCurveEditorApp creates all the subsystems and manages the dependencies. Specific "features" are implemented as members of this class.

What this addresses:
- Clear initilization order: gets rid of cyclic dependencies between FCurveEditorIntegrationExtension and FCurveEditorExtension.
- This effectively hides implementation details and avoids leaks of private API to public API: FSequencerCurveEditorApp is pimpl ptr in FCurveEditorExtension.
- Makes FCurveEditorExtension leaner... splitting the many responsibilities it had (god object) into smaller classes (= each one responsibility).

#Synced-CL 51334859

diff --git a/Engine/Source/Editor/Sequencer/Private/CurveEditor/Integration/CurveModelSyncer.cpp b/Engine/Source/Editor/Sequencer/Private/CurveEditor/Integration/CurveModelSyncer.cpp
new file mode 100644
--- /dev/null
+++ b/Engine/Source/Editor/Sequencer/Private/CurveEditor/Integration/CurveModelSyncer.cpp
@@ -0,0 +1,282 @@
+﻿// Copyright Epic Games, Inc. All Rights Reserved.
+
+#include "CurveModelSyncer.h"
+
+#include "CurveEditor.h"
+#include "Sequencer.h"
+#include "Filters/ViewModels/Linking/ILinkedFilterViewModel.h"
+#include "Modification/Utils/ScopedSelectionChange.h"
+#include "MVVM/SharedViewModelData.h"
+#include "MVVM/ViewModelPtr.h"
+#include "MVVM/Extensions/IOutlinerExtension.h"
+#include "MVVM/ViewModels/ViewModelIterators.h"
+#include "MVVM/Views/CurveEditor/CurveEditorWidgetOwner.h"
+#include "Tree/CurveEditorTree.h"
+
+namespace UE::Sequencer
+{
+FCurveModelSyncer::FCurveModelSyncer(
+	const TSharedRef<FSequencer>& InSequencer,
+	const TSharedRef<FCurveEditor>& InCurveEditor, const TSharedRef<ILinkedFilterViewModel>& InLinkedFilterViewModel
+	)
+	: WeakSequencer(InSequencer)
+	, CurveEditor(InCurveEditor)
+	, LinkedFilterViewModel(InLinkedFilterViewModel)
+{
+	const TSharedPtr<FSequenceModel> OwnerModel = InSequencer->GetViewModel()->GetRootSequenceModel();
+	FSimpleMulticastDelegate& HierarchyChanged = OwnerModel->GetSharedData()->SubscribeToHierarchyChanged(OwnerModel);
+	HierarchyChanged.AddRaw(this, &FCurveModelSyncer::OnHierarchyChanged);
+	
+	InSequencer->GetNodeTree()->OnUpdated().AddRaw(this, &FCurveModelSyncer::OnNodeTreeUpdated);
+}
+
+FCurveModelSyncer::~FCurveModelSyncer()
+{
+	if (const TSharedPtr<FSequencer> Sequencer = WeakSequencer.Pin())
+	{
+		const TSharedPtr<FSequenceModel> OwnerModel = Sequencer->GetViewModel()->GetRootSequenceModel();
+		OwnerModel->GetSharedData()->UnsubscribeFromHierarchyChanged(OwnerModel, this);
+	
+		Sequencer->GetNodeTree()->OnUpdated().RemoveAll(this);
+	}
+	
+	// Prevent crash on shutdown. RequestSyncSelection sets a timer that keeps a FScopedTransaction. 
+	// ~FScopedTransaction expects GEditor to be valid, but it is null at shutdown.
+	CleanUpPendingSyncSelection();
+}
+
+static bool bSyncSelectionRequested = false;
+void FCurveModelSyncer::RequestSyncSelection()
+{
+	const TSharedPtr<FSequencer> Sequencer = WeakSequencer.Pin();
+	if (bSyncSelectionRequested || !Sequencer)
+	{
+		return;
+	}
+	bSyncSelectionRequested = true;
+
+	// We schedule selection syncing to the next editor tick because we might want to select items that
+	// have just been added to the curve editor tree this tick. If it happened after the Slate update,
+	// these items don't yet have a UI widget, and so selecting them doesn't do anything.
+	//
+	// Note that we capture a weak pointer of our owner model because selection changes can happen
+	// right around the time when we want to unload everything (such as when loading a new map in the
+	// editor). We don't want to extend the lifetime of our stuff in that case.
+	TWeakPtr<FSequencerEditorViewModel> WeakRootViewModel = Sequencer->GetViewModel();
+
+	// Key selection supports undo. If RequestSyncSelection is called as part of an ongoing transaction, record the key selection change for undo.
+	const bool bShouldRecord = GUndo != nullptr;
+	// The current transaction needs to be extended until next tick...
+	TSharedPtr<FScopedTransaction> Transaction = bShouldRecord
+		? MakeShared<FScopedTransaction>(FText::GetEmpty(), bShouldRecord).ToSharedPtr() : nullptr;
+	// ... at which point we'll diff the changes.
+	TSharedPtr<CurveEditor::FScopedSelectionChange> SelectionChange = bShouldRecord
+		? MakeShared<CurveEditor::FScopedSelectionChange>(CurveEditor).ToSharedPtr() : nullptr;
+	
+	SyncSelectionHandle = GEditor->GetTimerManager()->SetTimerForNextTick(
+	[this, WeakRootViewModel, Transaction = MoveTemp(Transaction), SelectionChange = MoveTemp(SelectionChange)]() mutable
+		{
+			bSyncSelectionRequested = false;
+		
+			const TSharedPtr<FSequencerEditorViewModel> RootViewModel = WeakRootViewModel.Pin();
+			if (!RootViewModel.IsValid())
+			{
+				return;
+			}
+			const TSharedPtr<ISequencer> Sequencer = RootViewModel->GetSequencer();
+			if (!Sequencer)
+			{
+				return;
+			}
+				
+			SyncSelectionHandle.Invalidate();
+			// Notify UI to refresh the selection to that of Sequencer.
+			OnRequestSyncCurveEditorUIDelegate.Broadcast();
+
+			// Order matters: this appends a sub-transaction for the selection change...
+			SelectionChange.Reset();
+			// ... and only then can the parent transaction be closed. 
+			Transaction.Reset();
+		});
+}
+
+void FCurveModelSyncer::ResetCurveEditor()
+{
+	for (const TPair<TWeakPtr<FViewModel>, FCurveEditorTreeItemID>& Pair : ViewModelToTreeItemIDMap)
+	{
+		const FCurveEditorTreeItemID ItemID(Pair.Value);
+		const TViewModelPtr<ICurveEditorTreeItemExtension> ViewModel = CastViewModel<ICurveEditorTreeItemExtension>(Pair.Key.Pin());
+		if (ViewModel)
+		{
+			ViewModel->OnRemovedFromCurveEditor(CurveEditor);
+		}
+		CurveEditor->RemoveTreeItem(ItemID);
+	}
+
+	ViewModelToTreeItemIDMap.Reset();
+	TreeItemIDToViewModelMap.Reset();
+}
+
+TWeakPtr<FViewModel> FCurveModelSyncer::FindViewModelById(const FCurveEditorTreeItemID& InTreeId) const
+{
+	const TWeakPtr<FViewModel>* ViewModel = TreeItemIDToViewModelMap.Find(InTreeId);
+	return ViewModel ? *ViewModel : nullptr;
+}
+
+FCurveEditorTreeItemID FCurveModelSyncer::FindIdByViewModel(const TWeakPtr<FViewModel>& InViewModel) const
+{
+	const FCurveEditorTreeItemID* TreeItemID = ViewModelToTreeItemIDMap.Find(InViewModel);
+	return TreeItemID ? *TreeItemID : FCurveEditorTreeItemID::Invalid();
+}
+
+void FCurveModelSyncer::UpdateCurveEditor()
+{
+	const TSharedPtr<FSequencer> Sequencer = WeakSequencer.Pin();
+	if (!Sequencer)
+	{
+		return;
+	}
+	
+	// TODO UE-364008: In the future, we can remove the paths with bUseFilteredOut.
+	const bool bUseFilteredOut = Sequencer->GetSequencerSettings()->GetLinkFiltersWithCurveEditor();
+	const auto IsFilteredOut = [this, bUseFilteredOut](const TViewModelPtr<IOutlinerExtension>& InItem)
+	{
+		return UseLinkedFilteringUI()
+			? LinkedFilterViewModel->IsFilteredOut(InItem)
+			: bUseFilteredOut && InItem && InItem->IsFilteredOut();
+	};
+
+	FCurveEditorTree* CurveEditorTree = CurveEditor->GetTree();
+
+	// Guard against multiple broadcasts here and defer them until the end of this function
+	FScopedCurveEditorTreeEventGuard ScopedEventGuard = CurveEditorTree->ScopedEventGuard();
+
+	// Remove any stale tree items
+	for (auto It = ViewModelToTreeItemIDMap.CreateIterator(); It; ++It)
+	{
+		bool bIsRelevant = false;
+		const FCurveEditorTreeItemID ItemID(It.Value());
+		TViewModelPtr<ICurveEditorTreeItemExtension> ViewModel = CastViewModel<ICurveEditorTreeItemExtension>(It.Key().Pin());
+		if (ViewModel)
+		{
+			TViewModelPtr<IOutlinerExtension> OutlinerModel = ViewModel.ImplicitCast();
+			const bool bIsVisible = !IsFilteredOut(OutlinerModel);
+			
+			const FCurveEditorTreeItem* TreeItem = CurveEditorTree->FindItem(ItemID);
+
+			FCurveEditorTreeItemID ParentID;
+			if (auto Parent = ViewModel.AsModel()->FindAncestorOfType<ICurveEditorTreeItemExtension>())
+			{
+				ParentID = ViewModelToTreeItemIDMap.FindRef(Parent.AsModel());
+			}
+			
+			bIsRelevant = bIsVisible && TreeItem && TreeItem->GetParentID() == ParentID;
+		}
+		
+		if (!bIsRelevant)
+		{
+			if (ViewModel)
+			{
+				ViewModel->OnRemovedFromCurveEditor(CurveEditor);
+			}
+			CurveEditor->RemoveTreeItem(ItemID);
+			It.RemoveCurrent();
+		}
+	}
+
+	// Do a second pass to remove any items that were removed recursively above
+	for (auto It = ViewModelToTreeItemIDMap.CreateIterator(); It; ++It)
+	{
+		if (CurveEditorTree->FindItem(It->Value) == nullptr)
+		{
+			TreeItemIDToViewModelMap.Remove(It->Value);
+			It.RemoveCurrent();
+		}
+	}
+
+	// Iterate all non-filtered out outliners and check for curve editor tree extensions
+	const bool bIncludeRootNode = false;
+	bool bItemsAdded = false;
+	const TSharedPtr<FSequenceModel> OwnerModel = Sequencer->GetViewModel()->GetRootSequenceModel();
+	for (TParentFirstChildIterator<IOutlinerExtension> It(OwnerModel, bIncludeRootNode); It; ++It)
+	{
+		if (IsFilteredOut(*It))
+		{
+			It.IgnoreCurrentChildren();
+			continue;
+		}
+
+		if (TViewModelPtr<ICurveEditorTreeItemExtension> ChildViewModel = (*It).ImplicitCast())
+		{
+			const FCurveEditorTreeItemID ItemID = ViewModelToTreeItemIDMap.FindRef(ChildViewModel.AsModel());
+			if (!ItemID.IsValid())
+			{
+				bItemsAdded = true;
+				AddToCurveEditor(ChildViewModel, CurveEditor);
+			}
+		}
+	}
+
+	// If new items have been added, synchronize selection after the view models have been added to the curve editor. Synchronization depends on curve model tree item IDs
+	if (bItemsAdded)
+	{
+		RequestSyncSelection();
+	}
+}
+
+FCurveEditorTreeItemID FCurveModelSyncer::AddToCurveEditor(
+	TViewModelPtr<ICurveEditorTreeItemExtension> InViewModel,
+	TSharedPtr<FCurveEditor> InCurveEditor
+	)
+{
+	// If the view model doesn't want to be in the curve editor, bail out
+	// Note that this means we will create curve editor items for each parent in the hierarchy up
+	// until the first parent that doesn't implement ICurveEditorTreeItemExtension
+	// That is: we don't create "dummy" entries when there's a "gap" in the hierarchy
+	if (!InViewModel)
+	{
+		return FCurveEditorTreeItemID::Invalid();
+	}
+
+	// Check if we already have a valid curve editor ID
+	if (FCurveEditorTreeItemID* Existing = ViewModelToTreeItemIDMap.Find(InViewModel.AsModel()))
+	{
+		if (InCurveEditor->GetTree()->FindItem(*Existing) != nullptr)
+		{
+			return *Existing;
+		}
+	}
+
+	// Recursively create any needed parent curve editor items
+	TViewModelPtr<ICurveEditorTreeItemExtension> Parent = InViewModel.AsModel()->CastParent<ICurveEditorTreeItemExtension>();
+
+	FCurveEditorTreeItemID ParentID = Parent ?
+		AddToCurveEditor(Parent, InCurveEditor) :
+		FCurveEditorTreeItemID::Invalid();
+
+	// Create the new curve editor item
+	FCurveEditorTreeItem* NewItem = InCurveEditor->AddTreeItem(ParentID);
+	TSharedPtr<ICurveEditorTreeItem> CurveEditorTreeItem = InViewModel->GetCurveEditorTreeItem();
+	NewItem->SetWeakItem(CurveEditorTreeItem);
+	TOptional<FString> UniquePathName =  InViewModel->GetUniquePathName();
+	NewItem->SetUniquePathName(UniquePathName);
+	
+	// Register the new ID in our map and notify the view model
+	ViewModelToTreeItemIDMap.Add(InViewModel.AsModel(), NewItem->GetID());
+	TreeItemIDToViewModelMap.Add(NewItem->GetID(), InViewModel.AsModel());
+	
+	InViewModel->OnAddedToCurveEditor(NewItem->GetID(), InCurveEditor);
+
+	return NewItem->GetID();
+}
+
+void FCurveModelSyncer::CleanUpPendingSyncSelection()
+{
+	bSyncSelectionRequested = false;
+	
+	if (GEditor)
+	{
+		GEditor->GetTimerManager()->ClearTimer(SyncSelectionHandle);
+	}
+}
+} // namespace UE::Sequencer
diff --git a/Engine/Source/Editor/Sequencer/Private/CurveEditor/Integration/CurveModelSyncer.h b/Engine/Source/Editor/Sequencer/Private/CurveEditor/Integration/CurveModelSyncer.h
new file mode 100644
--- /dev/null
+++ b/Engine/Source/Editor/Sequencer/Private/CurveEditor/Integration/CurveModelSyncer.h
@@ -0,0 +1,84 @@
+﻿// Copyright Epic Games, Inc. All Rights Reserved.
+
+#pragma once
+
+#include "Containers/Map.h"
+#include "Engine/TimerHandle.h"
+#include "MVVM/Extensions/ICurveEditorTreeItemExtension.h"
+#include "MVVM/ViewModels/ViewModel.h"
+#include "Templates/UnrealTemplate.h"
+
+class FSequencer;
+
+namespace UE::Sequencer
+{
+class ILinkedFilterViewModel;
+
+/** 
+ * Makes sure that Curve Editor has the same selection as that in Sequencer.
+ * Consequently, it knows which FCurveEditorTreeItemID and FViewModel correspond to each other.
+ * 
+ * Responsible for creating FCurveModel instances in FCurveEditor based on ICurveEditorTreeItemExtension in the Sequencer hierarchy.
+ * This relies on outliner items implementing ICurveEditorTreeItemExtension (or its default shim) if they want to show up in the curve editor.
+ */
+class FCurveModelSyncer : public FNoncopyable
+{
+public:
+	
+	explicit FCurveModelSyncer(
+		const TSharedRef<FSequencer>& InSequencer,
+		const TSharedRef<FCurveEditor>& InCurveEditor, 
+		const TSharedRef<ILinkedFilterViewModel>& InLinkedFilterViewModel
+		);
+	~FCurveModelSyncer();
+	
+	/** Synchronize curve editor selection with sequencer outliner selection on the next update. */
+	void RequestSyncSelection();
+	
+	/** Clears the curve editor of all contents. */
+	void ResetCurveEditor();
+	
+	/** @return The view model associated with the tree item. */
+	TWeakPtr<FViewModel> FindViewModelById(const FCurveEditorTreeItemID& InTreeId) const;
+	/** @return The tree item associated with the view model. */
+	FCurveEditorTreeItemID FindIdByViewModel(const TWeakPtr<FViewModel>&) const;
+	
+	/** Keeps the curve editor items up-to-date with the sequencer outliner by adding/removing  entries as needed. */
+	void UpdateCurveEditor();
+	
+	/** Invoked to request that the UI be updated with the current selection. */
+	FSimpleMulticastDelegate& OnRequestSyncCurveEditorUI() { return OnRequestSyncCurveEditorUIDelegate; }
+	
+private:
+	
+	/** The owning Sequencer instance. Used to get settings. */
+	const TWeakPtr<FSequencer> WeakSequencer;
+	
+	/** The curve editor we are synchronizing to. */
+	const TSharedRef<FCurveEditor> CurveEditor;
+	/** Curve Editor selection must be refreshed when we switch filtering modes. */
+	const TSharedRef<ILinkedFilterViewModel> LinkedFilterViewModel;
+	
+	TMap<TWeakPtr<FViewModel>, FCurveEditorTreeItemID> ViewModelToTreeItemIDMap;
+	TMap<FCurveEditorTreeItemID, TWeakPtr<FViewModel>> TreeItemIDToViewModelMap;
+	
+	/** Handle to latent refresh of selection. */
+	FTimerHandle SyncSelectionHandle;
+	/** Invoked to request that the UI be updated with the current selection. */
+	FSimpleMulticastDelegate OnRequestSyncCurveEditorUIDelegate;
+	
+	/** Adds the given view-model to the curve editor */
+	FCurveEditorTreeItemID AddToCurveEditor(TViewModelPtr<ICurveEditorTreeItemExtension> InViewModel, TSharedPtr<FCurveEditor> InCurveEditor);
+	
+	/** Update curve editor items when sequencer outliner items change */
+	void OnHierarchyChanged() { UpdateCurveEditor(); }
+	/** Updates the curve editor as result of the node tree hierarchy changing. */
+	void OnNodeTreeUpdated() { UpdateCurveEditor(); }
+	
+	/** Cancels any pending transactions. */
+	void CleanUpPendingSyncSelection();
+};
+} // namespace UE::Sequencer
+
+
+
diff --git a/Engine/Source/Editor/Sequencer/Private/CurveEditor/SequencerCurveEditorApp.cpp b/Engine/Source/Editor/Sequencer/Private/CurveEditor/SequencerCurveEditorApp.cpp
new file mode 100644
--- /dev/null
+++ b/Engine/Source/Editor/Sequencer/Private/CurveEditor/SequencerCurveEditorApp.cpp
@@ -0,0 +1,113 @@
+﻿// Copyright Epic Games, Inc. All Rights Reserved.
+
+#include "SequencerCurveEditorApp.h"
+
+#include "Filters/FilterConfigIdentifiers.h"
+#include "Filters/ViewModels/Linking/LinkedFilterFactoryViewModel.h"
+#include "Filters/ViewModels/Linking/LinkedFilterViewModel.h"
+#include "Filters/ViewModels/Linking/MakeLinkedFilterViewModelArgs.h"
+#include "ICurveEditorModule.h"
+#include "MVVM/ViewModels/SequencerEditorViewModelImpl.h"
+#include "Specializations/SequencerCurveEditor.h"
+#include "Specializations/SequencerCurveEditorBounds.h"
+#include "Specializations/SequencerCurveEditorToolbarExtender.h"
+#include "Specializations/SequencerCurveModelUtils.h"
+#include "TimeSliderArgs.h"
+#include "MVVM/CurveEditorExtension.h"
+#include "Widgets/CurveEditor/SequencerCurveEditorTimeSliderController.h"
+
+namespace UE::Sequencer
+{
+class FCurveEditorExtension;
+
+namespace CurveEditorDetail
+{
+static FKeyAttributes GetDefaultKeyAttributes(USequencerSettings* InSettings)
+{
+	switch (InSettings->GetKeyInterpolation())
+	{
+	case EMovieSceneKeyInterpolation::User:     return FKeyAttributes().SetInterpMode(RCIM_Cubic).SetTangentMode(RCTM_User);
+	case EMovieSceneKeyInterpolation::Break:    return FKeyAttributes().SetInterpMode(RCIM_Cubic).SetTangentMode(RCTM_Break);
+	case EMovieSceneKeyInterpolation::Linear:   return FKeyAttributes().SetInterpMode(RCIM_Linear).SetTangentMode(RCTM_Auto);
+	case EMovieSceneKeyInterpolation::Constant: return FKeyAttributes().SetInterpMode(RCIM_Constant).SetTangentMode(RCTM_Auto);
+	case EMovieSceneKeyInterpolation::Auto:     return FKeyAttributes().SetInterpMode(RCIM_Cubic).SetTangentMode(RCTM_Auto);
+	case EMovieSceneKeyInterpolation::SmartAuto:
+	default:                                    return FKeyAttributes().SetInterpMode(RCIM_Cubic).SetTangentMode(RCTM_SmartAuto);
+	}
+}
+
+static TSharedRef<FCurveEditor> MakeCurveEditor(const TSharedRef<FSequencer>& InSequencer, const FTimeSliderArgs& InTimeSliderArgs)
+{
+	USequencerSettings* const SequencerSettings = InSequencer->GetSequencerSettings();
+	check(SequencerSettings);
+	
+	FCurveEditorInitParams CurveEditorInitParams;
+	CurveEditorInitParams.AdditionalEditorExtensions = { MakeShared<FSequencerCurveEditorToolbarExtender>(InSequencer.ToWeakPtr()) };
+	CurveEditorInitParams.ZoomScalingAttr.BindLambda([SequencerSettings]
+	{
+		return &SequencerSettings->GetCurveEditorZoomScaling();
+	});
+	CurveEditorInitParams.ResolveCurveModelDelegate.BindStatic(&ResolveCurveEditorModel, InSequencer.ToWeakPtr());
+		
+	const TSharedRef<FSequencerCurveEditor> CurveEditorModel = MakeShared<FSequencerCurveEditor>(InSequencer, InTimeSliderArgs.NumericTypeInterface);
+	CurveEditorModel->SetBounds(MakeUnique<FSequencerCurveEditorBounds>(InSequencer));
+	CurveEditorModel->InitCurveEditor(CurveEditorInitParams);
+
+	CurveEditorModel->InputSnapEnabledAttribute = MakeAttributeLambda([SequencerSettings] { return SequencerSettings->GetIsSnapEnabled(); });
+	CurveEditorModel->OnInputSnapEnabledChanged = FOnSetBoolean::CreateLambda([SequencerSettings](bool NewValue) { SequencerSettings->SetIsSnapEnabled(NewValue); });
+
+	CurveEditorModel->OutputSnapEnabledAttribute = MakeAttributeLambda([SequencerSettings] { return SequencerSettings->GetSnapCurveValueToInterval(); });
+	CurveEditorModel->OnOutputSnapEnabledChanged = FOnSetBoolean::CreateLambda([SequencerSettings](bool NewValue) { SequencerSettings->SetSnapCurveValueToInterval(NewValue); });
+
+	CurveEditorModel->FixedGridSpacingAttribute = MakeAttributeLambda([SequencerSettings]() -> TOptional<float> { return SequencerSettings->GetGridSpacing(); });
+	CurveEditorModel->InputSnapRateAttribute = MakeAttributeSP(InSequencer, &FSequencer::GetFocusedDisplayRate);
+
+	CurveEditorModel->DefaultKeyAttributes = MakeAttributeLambda([SequencerSettings]() { return GetDefaultKeyAttributes(SequencerSettings); });
+	
+	return CurveEditorModel;
+}
+
+static TSharedRef<FLinkedFilterViewModel> MakeFilteringModel(
+	const TSharedRef<FSequencer>& InSequencer, const FTimeSliderArgs& InTimeSliderArgs, const TSharedRef<FCurveEditor>& InCurveEditor
+	)
+{
+	// We create a custom Time Slider Controller which is just a wrapper around the actual one, but is 
+	// aware of our custom bounds logic. Currently the range the bar displays is tied to Sequencer 
+	// timeline and not the Bounds, so we need a way of changing it to look at the Bounds but only for 
+	// the Curve Editor time slider controller. We want everything else to just pass through though.
+	const TSharedRef<ITimeSliderController> CurveEditorTimeSliderController = MakeShared<FSequencerCurveEditorTimeSliderController>(
+		InTimeSliderArgs, InSequencer, InCurveEditor);
+	
+	const TSharedPtr<FLinkedFilterFactoryViewModel> FilterFactory = InSequencer->GetViewModelImpl()->GetLinkedFilterViewModelFactoryImpl();
+	const TSharedRef<FLinkedFilterViewModel> ViewModelImpl = FilterFactory->MakeFilteringModelImpl(
+		FMakeLinkedFilterViewModelArgs(ConfigIds::CurveEditor_InstancedConfigId)
+		);
+	
+	return ViewModelImpl;
+}
+}
+
+FSequencerCurveEditorApp::FSequencerCurveEditorApp(const TSharedRef<FSequencer>& InSequencer, const FTimeSliderArgs& InTimeSliderArgs)
+	: CurveEditor(CurveEditorDetail::MakeCurveEditor(InSequencer, InTimeSliderArgs))
+	, LinkedFilteringViewModel(CurveEditorDetail::MakeFilteringModel(InSequencer, InTimeSliderArgs, CurveEditor))
+	, SelectionSyncer(InSequencer, CurveEditor, LinkedFilteringViewModel)
+	, WidgetOwner(
+		InSequencer, CurveEditor, 
+		MakeShared<FSequencerCurveEditorTimeSliderController>(InTimeSliderArgs, InSequencer, CurveEditor), 
+		LinkedFilteringViewModel, SelectionSyncer
+		)
+{
+	SelectionSyncer.OnRequestSyncCurveEditorUI().AddLambda([this]
+	{
+		WidgetOwner.SyncSelection();
+	});
+	
+	SelectionSyncer.UpdateCurveEditor();
+}
+
+FSequencerCurveEditorApp* FSequencerCurveEditorApp::Get(const ISequencer& InSequencer)
+{
+	FCurveEditorExtension* CurveEditorExtension = InSequencer.GetViewModel()->GetRootModel()->CastDynamic<FCurveEditorExtension>();
+	return CurveEditorExtension ? CurveEditorExtension->CurveEditorApp.Get() : nullptr;
+}
+} // namespace UE::Sequencer
diff --git a/Engine/Source/Editor/Sequencer/Private/CurveEditor/SequencerCurveEditorApp.h b/Engine/Source/Editor/Sequencer/Private/CurveEditor/SequencerCurveEditorApp.h
new file mode 100644
--- /dev/null
+++ b/Engine/Source/Editor/Sequencer/Private/CurveEditor/SequencerCurveEditorApp.h
@@ -0,0 +1,78 @@
+﻿// Copyright Epic Games, Inc. All Rights Reserved.
+
+#pragma once
+
+#include "Integration/CurveModelSyncer.h"
+#include "MVVM/Views/CurveEditor/CurveEditorWidgetOwner.h"
+#include "Templates/SharedPointer.h"
+#include "Templates/UnrealTemplate.h"
+
+class ISequencer;
+
+namespace UE::Sequencer
+{
+class FLinkedFilterViewModel;
+
+/** 
+ * Integrates the Curve Editor into Sequencer.
+ * 
+ * Acts as the composition root for the Curve Editor within Sequencer:
+ * constructs all required subsystems, wires their dependencies,
+ * and coordinates their interaction with Sequencer.
+ * 
+ * FCurveEditorExtension provides the public API and forwards calls to this class.
+ * @see FCurveEditorExtension
+ */
+class FSequencerCurveEditorApp : public FNoncopyable
+{
+public:
+	
+	explicit FSequencerCurveEditorApp(
+		const TSharedRef<FSequencer>& InSequencer, const FTimeSliderArgs& InTimeSliderArgs
+		);
+	
+	/** 
+	 * @return The FSequencerCurveEditorApp used in the Sequencer instance. Returns nullptr, if curve editor integration is not enabled. 
+	 * @note This function exists as compromise / for convenience. Avoid introducing direct dependencies from core Sequencer code to the Curve Editor,
+	 * to ensure Sequencer follows the open-closed principle. For example, when adding new functionality, consider e.g. exposing a delegate in the
+	 * core Sequencer API and handling it within the Curve Editor integration layer. 
+	 */
+	static FSequencerCurveEditorApp* Get(const ISequencer& InSequencer);
+	
+	/** Gets the curve editor view-model */
+	TSharedPtr<FCurveEditor> GetCurveEditor() const { return CurveEditor; }
+
+	/** Opens the curve editor */
+	void OpenCurveEditor() { WidgetOwner.OpenCurveEditor(); }
+	/** Returns whether the curve editor is open */
+	bool IsCurveEditorOpen() const { return WidgetOwner.IsCurveEditorOpen(); }
+	/** Closes the curve editor */
+	void CloseCurveEditor() { WidgetOwner.CloseCurveEditor(); }
+
+	/** Curve editor tree widget */
+	TSharedPtr<SCurveEditorTree> GetCurveEditorTreeView() const { return WidgetOwner.GetTreeView(); }
+
+	/** Synchronize curve editor selection with sequencer outliner selection on the next update. */
+	void RequestSyncSelection() { SelectionSyncer.RequestSyncSelection(); }
+	/** Keeps the curve editor items up-to-date with the sequencer outliner by adding/removing  entries as needed. */
+	void UpdateCurveEditor() { SelectionSyncer.UpdateCurveEditor(); }
+	/** Clears the curve editor of all contents. */
+	void ResetCurveEditor() { SelectionSyncer.ResetCurveEditor(); }
+	
+private:
+	
+	/** The Curve Editor instance that Sequencer uses. */
+	TSharedRef<FCurveEditor> CurveEditor;
+	
+	/** Representation of Curve Editor's linked filters. */
+	TSharedRef<FLinkedFilterViewModel> LinkedFilteringViewModel;
+	
+	/** Handles that Curve Editor has the same selection as that of Sequencer. */
+	FCurveModelSyncer SelectionSyncer;
+	
+	/** Manages lifecycle of views exposed in Sequencer. */
+	FCurveEditorWidgetOwner WidgetOwner;
+};
+}
+
+
diff --git a/Engine/Source/Editor/Sequencer/Private/CurveEditor/Specializations/SequencerCurveEditor.h b/Engine/Source/Editor/Sequencer/Private/CurveEditor/Specializations/SequencerCurveEditor.h
new file mode 100644
--- /dev/null
+++ b/Engine/Source/Editor/Sequencer/Private/CurveEditor/Specializations/SequencerCurveEditor.h
@@ -0,0 +1,86 @@
+﻿// Copyright Epic Games, Inc. All Rights Reserved.
+
+#pragma once
+#include "CurveEditor.h"
+#include "CurveEditorAxis.h"
+#include "Sequencer.h"
+#include "SequencerTimeCurveEditorAxis.h"
+
+namespace UE::Sequencer
+{
+class FSequencerCurveEditor : public FCurveEditor
+{
+	TWeakPtr<FSequencer> WeakSequencer;
+	TSharedPtr<FLinearCurveEditorAxis> FocusedTimeAxis;
+	
+public:
+	
+	explicit FSequencerCurveEditor(TWeakPtr<FSequencer> InSequencer, TSharedPtr<INumericTypeInterface<double>> InNumericTypeInterface)
+		: WeakSequencer(InSequencer)
+	{
+		FocusedTimeAxis = MakeShared<FSequencerTimeCurveEditorAxis>(InSequencer);
+		FocusedTimeAxis->NumericTypeInterface = InNumericTypeInterface;
+
+		InSequencer.Pin()->OnActivateSequence().AddRaw(this, &FSequencerCurveEditor::HandleSequenceActivated);
+
+		AddAxis("FocusedSequenceTime", FocusedTimeAxis);
+	}
+
+	virtual ~FSequencerCurveEditor() override
+	{
+		if (TSharedPtr< FSequencer> Sequencer = WeakSequencer.Pin())
+		{
+			Sequencer->OnActivateSequence().RemoveAll(this);
+		}
+	}
+
+	virtual void GetGridLinesX(TArray<float>& MajorGridLines, TArray<float>& MinorGridLines, TArray<FText>* MajorGridLabels) const override
+	{
+		TSharedPtr<FSequencer> Sequencer = WeakSequencer.Pin();
+		if (!Sequencer)
+		{
+			return;
+		}
+		FCurveEditorScreenSpaceH PanelInputSpace = GetPanelInputSpace();
+
+		double MajorGridStep = 0.0;
+		int32  MinorDivisions = 0;
+		const bool bGotDimensions = Sequencer->GetGridMetrics(
+			PanelInputSpace.GetPhysicalWidth(), PanelInputSpace.GetInputMin(), PanelInputSpace.GetInputMax(), MajorGridStep, MinorDivisions
+			);
+
+		if (bGotDimensions)
+		{
+			const double FirstMajorLine = FMath::FloorToDouble(PanelInputSpace.GetInputMin() / MajorGridStep) * MajorGridStep;
+			const double LastMajorLine = FMath::CeilToDouble(PanelInputSpace.GetInputMax() / MajorGridStep) * MajorGridStep;
+
+			for (double CurrentMajorLine = FirstMajorLine; CurrentMajorLine < LastMajorLine; CurrentMajorLine += MajorGridStep)
+			{
+				MajorGridLines.Add(PanelInputSpace.SecondsToScreen(CurrentMajorLine));
+
+				for (int32 Step = 1; Step < MinorDivisions; ++Step)
+				{
+					MinorGridLines.Add(PanelInputSpace.SecondsToScreen(CurrentMajorLine + Step * MajorGridStep / MinorDivisions));
+				}
+			}
+		}
+	}
+
+	virtual int32 GetSupportedTangentTypes() override
+	{
+		return ((int32)ECurveEditorTangentTypes::InterpolationConstant	|
+			(int32)ECurveEditorTangentTypes::InterpolationLinear		|
+			(int32)ECurveEditorTangentTypes::InterpolationCubicAuto		|
+			(int32)ECurveEditorTangentTypes::InterpolationCubicUser		|
+			(int32)ECurveEditorTangentTypes::InterpolationCubicBreak	|
+			(int32)ECurveEditorTangentTypes::InterpolationCubicWeighted |
+			(int32)ECurveEditorTangentTypes::InterpolationCubicSmartAuto);
+	}
+
+	void HandleSequenceActivated(FMovieSceneSequenceIDRef) const
+	{
+		FocusedTimeAxis->NumericTypeInterface = WeakSequencer.Pin()->GetNumericTypeInterface();
+	}
+};
+} // namespace UE::Sequencer
+
diff --git a/Engine/Source/Editor/Sequencer/Private/CurveEditor/Specializations/SequencerCurveEditorBounds.h b/Engine/Source/Editor/Sequencer/Private/CurveEditor/Specializations/SequencerCurveEditorBounds.h
new file mode 100644
--- /dev/null
+++ b/Engine/Source/Editor/Sequencer/Private/CurveEditor/Specializations/SequencerCurveEditorBounds.h
@@ -0,0 +1,69 @@
+﻿// Copyright Epic Games, Inc. All Rights Reserved.
+
+#pragma once
+
+#include "AnimatedRange.h"
+#include "ICurveEditorBounds.h"
+#include "Sequencer.h"
+
+namespace UE::Sequencer
+{
+struct FSequencerCurveEditorBounds : ICurveEditorBounds
+{
+	FSequencerCurveEditorBounds(TSharedRef<FSequencer> InSequencer)
+		: WeakSequencer(InSequencer)
+	{
+		TRange<double> Bounds = InSequencer->GetViewRange();
+		InputMin = Bounds.GetLowerBoundValue();
+		InputMax = Bounds.GetUpperBoundValue();
+	}
+
+	virtual void GetInputBounds(double& OutMin, double& OutMax) const override
+	{
+		TSharedPtr<FSequencer> Sequencer = WeakSequencer.Pin();
+		if (Sequencer.IsValid())
+		{
+			const bool bLinkTimeRange = Sequencer->GetSequencerSettings()->GetLinkCurveEditorTimeRange();
+			if (bLinkTimeRange)
+			{
+				TRange<double> Bounds = Sequencer->GetViewRange();
+				OutMin = Bounds.GetLowerBoundValue();
+				OutMax = Bounds.GetUpperBoundValue();
+			}
+			else
+			{
+				// If they don't want to link the time range with Sequencer we return the cached value.
+				OutMin = InputMin;
+				OutMax = InputMax;
+			}
+		}
+	}
+
+	virtual void SetInputBounds(double InMin, double InMax) override
+	{
+		TSharedPtr<FSequencer> Sequencer = WeakSequencer.Pin();
+		if (Sequencer.IsValid())
+		{
+			const bool bLinkTimeRange = Sequencer->GetSequencerSettings()->GetLinkCurveEditorTimeRange();
+			if (bLinkTimeRange)
+			{
+				FFrameRate TickResolution = Sequencer->GetFocusedTickResolution();
+
+				if (InMin * TickResolution > TNumericLimits<int32>::Lowest() && InMax * TickResolution < TNumericLimits<int32>::Max())
+				{
+					Sequencer->SetViewRange(TRange<double>(InMin, InMax), EViewRangeInterpolation::Immediate);
+				}
+			}
+
+			// We update these even if you are linked to the Sequencer Timeline so that when you turn off the link setting
+			// you don't pop to your last values, instead your view stays as is and just stops moving when Sequencer moves.
+			InputMin = InMin;
+			InputMax = InMax;
+		}
+	}
+
+	/** The min/max values for the viewing range. Only used if Curve Editor/Sequencer aren't linked ranges. */
+	double InputMin, InputMax;
+	TWeakPtr<FSequencer> WeakSequencer;
+};
+} // namespace UE::Sequencer
diff --git a/Engine/Source/Editor/Sequencer/Private/CurveEditor/Specializations/SequencerCurveEditorToolbarExtender.h b/Engine/Source/Editor/Sequencer/Private/CurveEditor/Specializations/SequencerCurveEditorToolbarExtender.h
new file mode 100644
--- /dev/null
+++ b/Engine/Source/Editor/Sequencer/Private/CurveEditor/Specializations/SequencerCurveEditorToolbarExtender.h
@@ -0,0 +1,36 @@
+﻿// Copyright Epic Games, Inc. All Rights Reserved.
+
+#pragma once
+
+#include "ICurveEditorExtension.h"
+#include "Menus/SequencerToolbarUtils.h"
+#include "Sequencer.h"
+
+namespace UE::Sequencer
+{
+class FSequencerCurveEditorToolbarExtender : public ICurveEditorExtension
+{
+	TWeakPtr<FSequencer> WeakSequencer;
+public:
+
+	explicit FSequencerCurveEditorToolbarExtender(TWeakPtr<FSequencer> InWeakSequencer) : WeakSequencer(MoveTemp(InWeakSequencer)) {}
+	
+	virtual void BindCommands(TSharedRef<FUICommandList> CommandBindings) override {}
+	virtual TSharedPtr<FExtender> MakeToolbarExtender(const TSharedRef<FUICommandList>& InCommandList) override
+	{
+		TSharedRef<FExtender> Extender = MakeShared<FExtender>();
+		Extender->AddToolBarExtension("Adjustment", EExtensionHook::After, InCommandList,
+			FToolBarExtensionDelegate::CreateLambda([this](FToolBarBuilder& ToolbarBuilder)
+			{
+				TSharedPtr<FSequencer> SequenerPin = WeakSequencer.Pin();
+				
+				ToolbarBuilder.BeginSection("Keying");
+				ToolbarBuilder.PushCommandList(SequenerPin->GetCommandBindings().ToSharedRef());
+				AppendSequencerToolbarEntries(SequenerPin, ToolbarBuilder);
+				ToolbarBuilder.PopCommandList();
+				ToolbarBuilder.EndSection();
+			}));
+		return Extender;
+	}
+};
+} // namespace UE::Sequencer
diff --git a/Engine/Source/Editor/Sequencer/Private/CurveEditor/Specializations/SequencerCurveModelUtils.cpp b/Engine/Source/Editor/Sequencer/Private/CurveEditor/Specializations/SequencerCurveModelUtils.cpp
new file mode 100644
--- /dev/null
+++ b/Engine/Source/Editor/Sequencer/Private/CurveEditor/Specializations/SequencerCurveModelUtils.cpp
@@ -0,0 +1,51 @@
+﻿// Copyright Epic Games, Inc. All Rights Reserved.
+
+#include "SequencerCurveModelUtils.h"
+
+#include "Channels/MovieSceneChannelProxy.h"
+#include "ISequencerModule.h"
+#include "Modification/Resolution/CurveMetaDataIdentifiers.h"
+#include "Modules/ModuleManager.h"
+#include "MovieSceneSection.h"
+#include "Sequencer.h"
+
+namespace UE::Sequencer
+{
+TUniquePtr<FCurveModel> ResolveCurveEditorModel(const CurveEditor::FCurveModelLookUpArgs& InArgs, TWeakPtr<FSequencer> OwningSequencer)
+{
+	const TSharedPtr<ISequencer> SequencerPin = OwningSequencer.Pin();
+	if (!SequencerPin)
+	{
+		return nullptr;
+	}
+	
+	const CurveEditor::FCurveMetaDataIdentifiers* CurveMetadata = InArgs.GetCurveMetaData();
+	if (!CurveMetadata)
+	{
+		return nullptr;
+	}
+
+	UMovieSceneSection* OwningSection = Cast<UMovieSceneSection>(CurveMetadata->Owner.Get());
+	if (!OwningSection)
+	{
+		return nullptr;
+	}
+
+	const FMovieSceneChannelHandle ChannelHandle = OwningSection->GetChannelProxy().GetChannelByName(CurveMetadata->ChannelName);
+	ISequencerModule& SequencerModule = FModuleManager::LoadModuleChecked<ISequencerModule>("Sequencer");
+	ISequencerChannelInterface* EditorInterface = SequencerModule.FindChannelEditorInterface(ChannelHandle.GetChannelTypeName());
+	if (!EditorInterface)
+	{
+		return nullptr;
+	}
+	
+	const FMovieSceneChannelMetaData* ChannelMetadata = ChannelHandle.GetMetaData();
+	FCreateCurveEditorModelParams Params = {
+		OwningSection,
+		ChannelMetadata->WeakOwningObject.Get(),
+		SequencerPin.ToSharedRef()
+	};
+
+	return EditorInterface->CreateCurveEditorModel_Raw(ChannelHandle, Params);
+}
+} // namespace UE::Sequencer
\ No newline at end of file
diff --git a/Engine/Source/Editor/Sequencer/Private/CurveEditor/Specializations/SequencerCurveModelUtils.h b/Engine/Source/Editor/Sequencer/Private/CurveEditor/Specializations/SequencerCurveModelUtils.h
new file mode 100644
--- /dev/null
+++ b/Engine/Source/Editor/Sequencer/Private/CurveEditor/Specializations/SequencerCurveModelUtils.h
@@ -0,0 +1,16 @@
+﻿// Copyright Epic Games, Inc. All Rights Reserved.
+
+#pragma once
+
+#include "Modification/Resolution/CurveModelLookUpInfo.h"
+#include "Templates/UniquePtr.h"
+
+class FCurveModel;
+class FSequencer;
+namespace UE::CurveEditor { struct FCurveModelLookUpArgs; }
+
+namespace UE::Sequencer
+{
+/** Constructs a FCurveModel by looking up the channel from the owning UMovieSceneSection and ChannelName. */
+TUniquePtr<FCurveModel> ResolveCurveEditorModel(const CurveEditor::FCurveModelLookUpArgs& InArgs, TWeakPtr<FSequencer> OwningSequencer);
+} // namespace UE::Sequencer
diff --git a/Engine/Source/Editor/Sequencer/Private/CurveEditor/Specializations/SequencerTimeCurveEditorAxis.h b/Engine/Source/Editor/Sequencer/Private/CurveEditor/Specializations/SequencerTimeCurveEditorAxis.h
new file mode 100644
--- /dev/null
+++ b/Engine/Source/Editor/Sequencer/Private/CurveEditor/Specializations/SequencerTimeCurveEditorAxis.h
@@ -0,0 +1,73 @@
+﻿// Copyright Epic Games, Inc. All Rights Reserved.
+
+#pragma once
+
+#include "CurveEditorAxis.h"
+#include "SCurveEditorView.h"
+#include "Sequencer.h"
+
+namespace UE::Sequencer
+{
+/** Custom curve editor axis that displays the 'current time' in display rate */
+class FSequencerTimeCurveEditorAxis : public FLinearCurveEditorAxis
+{
+	TWeakPtr<FSequencer> WeakSequencer;
+public:
+
+	explicit FSequencerTimeCurveEditorAxis(TWeakPtr<FSequencer> InWeakSequencer)
+		: WeakSequencer(InWeakSequencer)
+	{}
+
+	virtual void GetGridLines(
+		const FCurveEditor& CurveEditor, const SCurveEditorView& View, FCurveEditorViewAxisID AxisID, 
+		TArray<double>& OutMajorGridLines, TArray<double>& OutMinorGridLines, ECurveEditorAxisOrientation Axis
+		) const override
+	{
+		TSharedPtr<FSequencer> Sequencer = WeakSequencer.Pin();
+
+		if (!Sequencer.IsValid())
+		{
+			return;
+		}
+
+		double MajorGridStep = 0.0;
+		int32  MinorDivisions = 0;
+		
+		float Size = 1.0;
+		float Min  = 0.0;
+		float Max  = 1.0;
+
+		if (Axis == ECurveEditorAxisOrientation::Horizontal)
+		{
+			FCurveEditorScreenSpaceH AxisSpace = View.GetHorizontalAxisSpace(AxisID);
+			Size = AxisSpace.GetPhysicalWidth();
+			Min  = AxisSpace.GetInputMin();
+			Max  = AxisSpace.GetInputMax();
+		}
+		else
+		{
+			FCurveEditorScreenSpaceV AxisSpace = View.GetVerticalAxisSpace(AxisID);
+			Size = AxisSpace.GetPhysicalHeight();
+			Min = AxisSpace.GetOutputMin();
+			Max = AxisSpace.GetOutputMax();
+		}
+
+		if (Sequencer.IsValid() && Sequencer->GetGridMetrics(Size, Min, Max, MajorGridStep, MinorDivisions))
+		{
+			double FirstMajorLine = FMath::FloorToDouble(Min / MajorGridStep) * MajorGridStep;
+			double LastMajorLine = FMath::CeilToDouble(Max / MajorGridStep) * MajorGridStep;
+
+			for (double CurrentMajorLine = FirstMajorLine; CurrentMajorLine < LastMajorLine; CurrentMajorLine += MajorGridStep)
+			{
+				OutMajorGridLines.Add(CurrentMajorLine);
+
+				for (int32 Step = 1; Step < MinorDivisions; ++Step)
+				{
+					double MinorLine = CurrentMajorLine + Step * MajorGridStep / MinorDivisions;
+					OutMinorGridLines.Add(MinorLine);
+				}
+			}
+		}
+	}
+};
+} // namespace UE::Sequencer
diff --git a/Engine/Source/Editor/Sequencer/Private/MVVM/CurveEditorExtension.cpp b/Engine/Source/Editor/Sequencer/Private/MVVM/CurveEditorExtension.cpp
--- a/Engine/Source/Editor/Sequencer/Private/MVVM/CurveEditorExtension.cpp
+++ b/Engine/Source/Editor/Sequencer/Private/MVVM/CurveEditorExtension.cpp
@@ -1,50 +1,18 @@
 // Copyright Epic Games, Inc. All Rights Reserved.
 
 #include "MVVM/CurveEditorExtension.h"
-#include "MVVM/Selection/Selection.h"
 
-#include "FrameNumberDetailsCustomization.h"
-#include "Filters/FilterConfigIdentifiers.h"
-#include "Filters/SCurveEditorFilterPanel.h"
+#include "CurveEditor/SequencerCurveEditorApp.h"
 #include "Filters/ViewModels/Linking/ILinkedFilterViewModel.h"
-#include "Filters/ViewModels/Linking/LinkedFilterFactoryViewModel.h"
-#include "Filters/ViewModels/Linking/LinkedFilterViewModel.h"
-#include "Filters/ViewModels/Linking/MakeLinkedFilterViewModelArgs.h"
+#include "FrameNumberDetailsCustomization.h"
 #include "Framework/Docking/TabManager.h"
-#include "IPropertyRowGenerator.h"
-#include "IStructureDetailsView.h"
-#include "MVVM/ViewModels/ViewModelIterators.h"
 #include "MVVM/ViewModels/SequencerEditorViewModel.h"
-#include "MVVM/ViewModels/ChannelModel.h"
-#include "MVVM/Extensions/IOutlinerExtension.h"
-#include "CurveEditorAxis.h"
-#include "ICurveEditorExtension.h"
-#include "SCurveEditorView.h"
+#include "Modification/Resolution/CurveMetaDataIdentifiers.h"
 #include "SCurveEditorPanel.h"
-#include "SCurveEditorToolProperties.h"
 #include "SCurveKeyDetailPanel.h"
-#include "SPopoutTabInlineContent.h"
-#include "SSequencerTreeFilterStatusBar.h"
-#include "STemporarilyFocusedSpinBox.h"
 #include "Sequencer.h"
-#include "SequencerCommands.h"
-#include "Menus/SequencerToolbarUtils.h"
-#include "Modification/Resolution/CurveMetaDataIdentifiers.h"
-#include "Modification/Resolution/CurveModelLookUpInfo.h"
-#include "Modification/Utils/ScopedSelectionChange.h"
 #include "Toolkits/IToolkitHost.h"
 #include "Tree/SCurveEditorTree.h"
-#include "Tree/SCurveEditorTreeFilterStatusBar.h"
-#include "Tree/SCurveEditorTreeTextFilter.h"
-#include "Widgets/CurveEditor/SSequencerCurveEditor.h"
-#include "Widgets/CurveEditor/SequencerCurveEditorTimeSliderController.h"
-#include "Widgets/Docking/SDockTab.h"
-#include "Widgets/Input/SButton.h"
-#include "Widgets/Layout/SScrollBorder.h"
-#include "WorkspaceMenuStructure.h"
-#include "WorkspaceMenuStructureModule.h"
-#include "ViewModels/SequencerEditorViewModelImpl.h"
-#include "Views/CurveEditor/CurveEditorWidgetOwner.h"
 
 #define LOCTEXT_NAMESPACE "SequencerCurveEditorExtension"
 
@@ -52,274 +20,13 @@ namespace UE
 {
 namespace Sequencer
 {
-
-/** Custom curve editor axis that displays the 'current time' in display rate */
-class FSequencerTimeCurveEditorAxis : public FLinearCurveEditorAxis
-{
-public:
-
-	TWeakPtr<FSequencer> WeakSequencer;
-	FSequencerTimeCurveEditorAxis(TWeakPtr<FSequencer> InWeakSequencer)
-		: WeakSequencer(InWeakSequencer)
-	{
-	}
-
-	void GetGridLines(const FCurveEditor& CurveEditor, const SCurveEditorView& View, FCurveEditorViewAxisID AxisID, TArray<double>& OutMajorGridLines, TArray<double>& OutMinorGridLines, ECurveEditorAxisOrientation Axis) const override
-	{
-		TSharedPtr<FSequencer> Sequencer = WeakSequencer.Pin();
-
-		if (!Sequencer.IsValid())
-		{
-			return;
-		}
-
-		double ToSeconds = Sequencer->GetFocusedTickResolution().AsInterval();
-
-		double MajorGridStep = 0.0;
-		int32  MinorDivisions = 0;
-
-
-		float Size = 1.0;
-		float Min  = 0.0;
-		float Max  = 1.0;
-
-		if (Axis == ECurveEditorAxisOrientation::Horizontal)
-		{
-			FCurveEditorScreenSpaceH AxisSpace = View.GetHorizontalAxisSpace(AxisID);
-			Size = AxisSpace.GetPhysicalWidth();
-			Min  = AxisSpace.GetInputMin();
-			Max  = AxisSpace.GetInputMax();
-		}
-		else
-		{
-			FCurveEditorScreenSpaceV AxisSpace = View.GetVerticalAxisSpace(AxisID);
-			Size = AxisSpace.GetPhysicalHeight();
-			Min = AxisSpace.GetOutputMin();
-			Max = AxisSpace.GetOutputMax();
-		}
-
-		if (Sequencer.IsValid() && Sequencer->GetGridMetrics(Size, Min, Max, MajorGridStep, MinorDivisions))
-		{
-			double FirstMajorLine = FMath::FloorToDouble(Min / MajorGridStep) * MajorGridStep;
-			double LastMajorLine = FMath::CeilToDouble(Max / MajorGridStep) * MajorGridStep;
-
-			for (double CurrentMajorLine = FirstMajorLine; CurrentMajorLine < LastMajorLine; CurrentMajorLine += MajorGridStep)
-			{
-				OutMajorGridLines.Add(CurrentMajorLine);
-
-				for (int32 Step = 1; Step < MinorDivisions; ++Step)
-				{
-					double MinorLine = CurrentMajorLine + Step * MajorGridStep / MinorDivisions;
-					OutMinorGridLines.Add(MinorLine);
-				}
-			}
-		}
-	}
-};
-
-class FSequencerCurveEditor : public FCurveEditor
-{
-public:
-	TWeakPtr<FSequencer> WeakSequencer;
-	TSharedPtr<FLinearCurveEditorAxis> FocusedTimeAxis;
-
-	FSequencerCurveEditor(TWeakPtr<FSequencer> InSequencer, TSharedPtr<INumericTypeInterface<double>> InNumericTypeInterface)
-		: WeakSequencer(InSequencer)
-	{
-		FocusedTimeAxis = MakeShared<FSequencerTimeCurveEditorAxis>(InSequencer);
-		FocusedTimeAxis->NumericTypeInterface = InNumericTypeInterface;
-
-		InSequencer.Pin()->OnActivateSequence().AddRaw(this, &FSequencerCurveEditor::HandleSequenceActivated);
-
-		AddAxis("FocusedSequenceTime", FocusedTimeAxis);
-	}
-
-	~FSequencerCurveEditor()
-	{
-		if (TSharedPtr< FSequencer> Sequencer = WeakSequencer.Pin())
-		{
-			Sequencer->OnActivateSequence().RemoveAll(this);
-		}
-	}
-
-	virtual void GetGridLinesX(TArray<float>& MajorGridLines, TArray<float>& MinorGridLines, TArray<FText>* MajorGridLabels) const override
-	{
-		TSharedPtr<FSequencer> Sequencer = WeakSequencer.Pin();
-		FCurveEditorScreenSpaceH PanelInputSpace = GetPanelInputSpace();
-
-		double MajorGridStep = 0.0;
-		int32  MinorDivisions = 0;
-
-		if (Sequencer.IsValid() && Sequencer->GetGridMetrics(PanelInputSpace.GetPhysicalWidth(), PanelInputSpace.GetInputMin(), PanelInputSpace.GetInputMax(), MajorGridStep, MinorDivisions))
-		{
-			const double FirstMajorLine = FMath::FloorToDouble(PanelInputSpace.GetInputMin() / MajorGridStep) * MajorGridStep;
-			const double LastMajorLine = FMath::CeilToDouble(PanelInputSpace.GetInputMax() / MajorGridStep) * MajorGridStep;
-
-			for (double CurrentMajorLine = FirstMajorLine; CurrentMajorLine < LastMajorLine; CurrentMajorLine += MajorGridStep)
-			{
-				MajorGridLines.Add(PanelInputSpace.SecondsToScreen(CurrentMajorLine));
-
-				for (int32 Step = 1; Step < MinorDivisions; ++Step)
-				{
-					MinorGridLines.Add(PanelInputSpace.SecondsToScreen(CurrentMajorLine + Step * MajorGridStep / MinorDivisions));
-				}
-			}
-		}
-	}
-	int32 GetSupportedTangentTypes() override
-	{
-		return ((int32)ECurveEditorTangentTypes::InterpolationConstant	|
-			(int32)ECurveEditorTangentTypes::InterpolationLinear		|
-			(int32)ECurveEditorTangentTypes::InterpolationCubicAuto		|
-			(int32)ECurveEditorTangentTypes::InterpolationCubicUser		|
-			(int32)ECurveEditorTangentTypes::InterpolationCubicBreak	|
-			(int32)ECurveEditorTangentTypes::InterpolationCubicWeighted |
-			(int32)ECurveEditorTangentTypes::InterpolationCubicSmartAuto);
-	}
-
-	void HandleSequenceActivated(FMovieSceneSequenceIDRef NewSequenceID)
-	{
-		FocusedTimeAxis->NumericTypeInterface = WeakSequencer.Pin()->GetNumericTypeInterface();
-	}
-};
-
-struct FSequencerCurveEditorBounds : ICurveEditorBounds
-{
-	FSequencerCurveEditorBounds(TSharedRef<FSequencer> InSequencer)
-		: WeakSequencer(InSequencer)
-	{
-		TRange<double> Bounds = InSequencer->GetViewRange();
-		InputMin = Bounds.GetLowerBoundValue();
-		InputMax = Bounds.GetUpperBoundValue();
-	}
-
-	virtual void GetInputBounds(double& OutMin, double& OutMax) const override
-	{
-		TSharedPtr<FSequencer> Sequencer = WeakSequencer.Pin();
-		if (Sequencer.IsValid())
-		{
-			const bool bLinkTimeRange = Sequencer->GetSequencerSettings()->GetLinkCurveEditorTimeRange();
-			if (bLinkTimeRange)
-			{
-				TRange<double> Bounds = Sequencer->GetViewRange();
-				OutMin = Bounds.GetLowerBoundValue();
-				OutMax = Bounds.GetUpperBoundValue();
-			}
-			else
-			{
-				// If they don't want to link the time range with Sequencer we return the cached value.
-				OutMin = InputMin;
-				OutMax = InputMax;
-			}
-		}
-	}
-
-	virtual void SetInputBounds(double InMin, double InMax) override
-	{
-		TSharedPtr<FSequencer> Sequencer = WeakSequencer.Pin();
-		if (Sequencer.IsValid())
-		{
-			const bool bLinkTimeRange = Sequencer->GetSequencerSettings()->GetLinkCurveEditorTimeRange();
-			if (bLinkTimeRange)
-			{
-				FFrameRate TickResolution = Sequencer->GetFocusedTickResolution();
-
-				if (InMin * TickResolution > TNumericLimits<int32>::Lowest() && InMax * TickResolution < TNumericLimits<int32>::Max())
-				{
-					Sequencer->SetViewRange(TRange<double>(InMin, InMax), EViewRangeInterpolation::Immediate);
-				}
-			}
-
-			// We update these even if you are linked to the Sequencer Timeline so that when you turn off the link setting
-			// you don't pop to your last values, instead your view stays as is and just stops moving when Sequencer moves.
-			InputMin = InMin;
-			InputMax = InMax;
-		}
-	}
-
-	/** The min/max values for the viewing range. Only used if Curve Editor/Sequencer aren't linked ranges. */
-	double InputMin, InputMax;
-	TWeakPtr<FSequencer> WeakSequencer;
-};
-
-class FSequencerCurveEditorToolbarExtender : public ICurveEditorExtension
-{
-	TWeakPtr<FSequencer> WeakSequencer;
-public:
-
-	explicit FSequencerCurveEditorToolbarExtender(TWeakPtr<FSequencer> InWeakSequencer) : WeakSequencer(MoveTemp(InWeakSequencer)) {}
-	
-	virtual void BindCommands(TSharedRef<FUICommandList> CommandBindings) override {}
-	virtual TSharedPtr<FExtender> MakeToolbarExtender(const TSharedRef<FUICommandList>& InCommandList) override
-	{
-		TSharedRef<FExtender> Extender = MakeShared<FExtender>();
-		Extender->AddToolBarExtension("Adjustment", EExtensionHook::After, InCommandList,
-			FToolBarExtensionDelegate::CreateLambda([this](FToolBarBuilder& ToolbarBuilder)
-			{
-				TSharedPtr<FSequencer> SequenerPin = WeakSequencer.Pin();
-				
-				ToolbarBuilder.BeginSection("Keying");
-				ToolbarBuilder.PushCommandList(SequenerPin->GetCommandBindings().ToSharedRef());
-				AppendSequencerToolbarEntries(SequenerPin, ToolbarBuilder);
-				ToolbarBuilder.PopCommandList();
-				ToolbarBuilder.EndSection();
-			}));
-		return Extender;
-	}
-};
-
-/** Constructs a FCurveModel by looking up the channel from the owning UMovieSceneSection and ChannelName. */
-static TUniquePtr<FCurveModel> ResolveCurveEditorModel(const CurveEditor::FCurveModelLookUpArgs& InArgs, TWeakPtr<FSequencer> OwningSequencer)
-{
-	const TSharedPtr<ISequencer> SequencerPin = OwningSequencer.Pin();
-	if (!SequencerPin)
-	{
-		return nullptr;
-	}
-	
-	const CurveEditor::FCurveMetaDataIdentifiers* CurveMetadata = InArgs.GetCurveMetaData();
-	if (!CurveMetadata)
-	{
-		return nullptr;
-	}
-
-	UMovieSceneSection* OwningSection = Cast<UMovieSceneSection>(CurveMetadata->Owner.Get());
-	if (!OwningSection)
-	{
-		return nullptr;
-	}
-
-	const FMovieSceneChannelHandle ChannelHandle = OwningSection->GetChannelProxy().GetChannelByName(CurveMetadata->ChannelName);
-	ISequencerModule& SequencerModule = FModuleManager::LoadModuleChecked<ISequencerModule>("Sequencer");
-	ISequencerChannelInterface* EditorInterface = SequencerModule.FindChannelEditorInterface(ChannelHandle.GetChannelTypeName());
-	if (!EditorInterface)
-	{
-		return nullptr;
-	}
-	
-	const FMovieSceneChannelMetaData* ChannelMetadata = ChannelHandle.GetMetaData();
-	FCreateCurveEditorModelParams Params = {
-		OwningSection,
-		ChannelMetadata->WeakOwningObject.Get(),
-		SequencerPin.ToSharedRef()
-	};
-
-	return EditorInterface->CreateCurveEditorModel_Raw(ChannelHandle, Params);
-}
-
 const FName FCurveEditorExtension::CurveEditorTabName = FName(TEXT("SequencerGraphEditor"));
 
 FCurveEditorExtension::FCurveEditorExtension()
-{
-	
-}
+{}
 
 FCurveEditorExtension::~FCurveEditorExtension()
-{
-	// Prevent crash on shutdown. RequestSyncSelection sets a timer that keeps a FScopedTransaction. 
-	// ~FScopedTransaction expects GEditor to be valid, but it is null at shutdown.
-	CleanUpPendingSyncSelection();
-}
+{}
 
 TSharedPtr<ISequencer> FCurveEditorExtension::GetSequencer() const
 {
@@ -380,163 +87,48 @@ void FCurveEditorExtension::CreateCurveEditor(const FTimeSliderArgs& TimeSliderA
 		return;
 	}
 
-	// Create the curve editor;
-	{
-		FCurveEditorInitParams CurveEditorInitParams;
-		CurveEditorInitParams.AdditionalEditorExtensions = { MakeShared<FSequencerCurveEditorToolbarExtender>(Sequencer.ToWeakPtr()) };
-		CurveEditorInitParams.ZoomScalingAttr.BindLambda([SequencerSettings]
-		{
-			return &SequencerSettings->GetCurveEditorZoomScaling();
-		});
-		CurveEditorInitParams.ResolveCurveModelDelegate.BindStatic(&ResolveCurveEditorModel, Sequencer.ToWeakPtr());
-		
-		CurveEditorModel = MakeShared<FSequencerCurveEditor>(Sequencer, TimeSliderArgs.NumericTypeInterface);
-		CurveEditorModel->SetBounds(MakeUnique<FSequencerCurveEditorBounds>(Sequencer.ToSharedRef()));
-		CurveEditorModel->InitCurveEditor(CurveEditorInitParams);
-
-		CurveEditorModel->InputSnapEnabledAttribute = MakeAttributeLambda([SequencerSettings] { return SequencerSettings->GetIsSnapEnabled(); });
-		CurveEditorModel->OnInputSnapEnabledChanged = FOnSetBoolean::CreateLambda([SequencerSettings](bool NewValue) { SequencerSettings->SetIsSnapEnabled(NewValue); });
-
-		CurveEditorModel->OutputSnapEnabledAttribute = MakeAttributeLambda([SequencerSettings] { return SequencerSettings->GetSnapCurveValueToInterval(); });
-		CurveEditorModel->OnOutputSnapEnabledChanged = FOnSetBoolean::CreateLambda([SequencerSettings](bool NewValue) { SequencerSettings->SetSnapCurveValueToInterval(NewValue); });
-
-		CurveEditorModel->FixedGridSpacingAttribute = MakeAttributeLambda([SequencerSettings]() -> TOptional<float> { return SequencerSettings->GetGridSpacing(); });
-		CurveEditorModel->InputSnapRateAttribute = MakeAttributeSP(Sequencer.Get(), &FSequencer::GetFocusedDisplayRate);
+	CurveEditorApp = MakePimpl<FSequencerCurveEditorApp>(Sequencer.ToSharedRef(), TimeSliderArgs);
+}
 
-		CurveEditorModel->DefaultKeyAttributes = MakeAttributeLambda([this]() { return GetDefaultKeyAttributes(); });
-	}
-	
-	// We create a custom Time Slider Controller which is just a wrapper around the actual one, but is 
-	// aware of our custom bounds logic. Currently the range the bar displays is tied to Sequencer 
-	// timeline and not the Bounds, so we need a way of changing it to look at the Bounds but only for 
-	// the Curve Editor time slider controller. We want everything else to just pass through though.
-	const TSharedRef<ITimeSliderController> CurveEditorTimeSliderController = MakeShared<FSequencerCurveEditorTimeSliderController>(
-		TimeSliderArgs, Sequencer, CurveEditorModel.ToSharedRef());
-	
-	const TSharedPtr<FLinkedFilterFactoryViewModel> FilterFactory = Sequencer->GetViewModelImpl()->GetLinkedFilterViewModelFactoryImpl();
-	const TSharedRef<FLinkedFilterViewModel> ViewModelImpl = FilterFactory->MakeFilteringModelImpl(
-		FMakeLinkedFilterViewModelArgs(ConfigIds::CurveEditor_InstancedConfigId)
-		);
-	CurveEditorFilterClient = ViewModelImpl;
-	WidgetOwner = MakePimpl<FCurveEditorWidgetOwner>(
-		Sequencer.ToSharedRef(), CurveEditorModel.ToSharedRef(), CurveEditorTimeSliderController, ViewModelImpl
-		);
-	
-	// TODO UE-364009: Effectively, this makes FCurveEditorIntegrationExtension bind to FilterBarProxy. In the future, refactor 
-	// FCurveEditorIntegrationExtension (non-extension) to be a separate class that FCurveEditorExtension injects the dependencies to.
-	OnCurveEditorCreated.Broadcast();
+TSharedPtr<FCurveEditor> FCurveEditorExtension::GetCurveEditor() const
+{
+	return CurveEditorApp ? CurveEditorApp->GetCurveEditor() : nullptr;
 }
 
 void FCurveEditorExtension::OpenCurveEditor()
 {
-	if (WidgetOwner)
+	if (CurveEditorApp)
 	{
-		WidgetOwner->OpenCurveEditor();
+		CurveEditorApp->OpenCurveEditor();
 	}
 }
 
 bool FCurveEditorExtension::IsCurveEditorOpen() const
 {
-	return WidgetOwner && WidgetOwner->IsCurveEditorOpen();
+	return CurveEditorApp && CurveEditorApp->IsCurveEditorOpen();
 }
 
 void FCurveEditorExtension::CloseCurveEditor()
 {
-	if (WidgetOwner)
+	if (CurveEditorApp)
 	{
-		WidgetOwner->CloseCurveEditor();
+		CurveEditorApp->CloseCurveEditor();
 	}
 }
 
 TSharedPtr<SCurveEditorTree> FCurveEditorExtension::GetCurveEditorTreeView() const
 {
-	return WidgetOwner ? WidgetOwner->GetTreeView() : nullptr;
-}
-
-FKeyAttributes FCurveEditorExtension::GetDefaultKeyAttributes() const
-{
-	USequencerSettings* const Settings = GetSequencerSettings();
-	check(Settings);
-
-	switch (Settings->GetKeyInterpolation())
-	{
-	case EMovieSceneKeyInterpolation::User:     return FKeyAttributes().SetInterpMode(RCIM_Cubic).SetTangentMode(RCTM_User);
-	case EMovieSceneKeyInterpolation::Break:    return FKeyAttributes().SetInterpMode(RCIM_Cubic).SetTangentMode(RCTM_Break);
-	case EMovieSceneKeyInterpolation::Linear:   return FKeyAttributes().SetInterpMode(RCIM_Linear).SetTangentMode(RCTM_Auto);
-	case EMovieSceneKeyInterpolation::Constant: return FKeyAttributes().SetInterpMode(RCIM_Constant).SetTangentMode(RCTM_Auto);
-	case EMovieSceneKeyInterpolation::Auto:     return FKeyAttributes().SetInterpMode(RCIM_Cubic).SetTangentMode(RCTM_Auto);
-	case EMovieSceneKeyInterpolation::SmartAuto:
-	default:                                    return FKeyAttributes().SetInterpMode(RCIM_Cubic).SetTangentMode(RCTM_SmartAuto);
-	}
+	return CurveEditorApp ? CurveEditorApp->GetCurveEditorTreeView() : nullptr;
 }
 
-static bool bSyncSelectionRequested = false;
 void FCurveEditorExtension::RequestSyncSelection()
 {
-	if (bSyncSelectionRequested || !WidgetOwner)
+	if (CurveEditorApp)
 	{
-		return;
+		CurveEditorApp->RequestSyncSelection();
 	}
-
-	bSyncSelectionRequested = true;
-
-	// We schedule selection syncing to the next editor tick because we might want to select items that
-	// have just been added to the curve editor tree this tick. If it happened after the Slate update,
-	// these items don't yet have a UI widget, and so selecting them doesn't do anything.
-	//
-	// Note that we capture a weak pointer of our owner model because selection changes can happen
-	// right around the time when we want to unload everything (such as when loading a new map in the
-	// editor). We don't want to extend the lifetime of our stuff in that case.
-	TWeakPtr<FSequencerEditorViewModel> WeakRootViewModel(WeakOwnerModel);
-
-	// Key selection supports undo. If RequestSyncSelection is called as part of an ongoing transaction, record the key selection change for undo.
-	const bool bShouldRecord = GUndo != nullptr;
-	// The current transaction needs to be extended until next tick...
-	TSharedPtr<FScopedTransaction> Transaction = bShouldRecord
-		? MakeShared<FScopedTransaction>(FText::GetEmpty(), bShouldRecord).ToSharedPtr() : nullptr;
-	// ... at which point we'll diff the changes.
-	TSharedPtr<CurveEditor::FScopedSelectionChange> SelectionChange = bShouldRecord
-		? MakeShared<CurveEditor::FScopedSelectionChange>(CurveEditorModel).ToSharedPtr() : nullptr;
-	
-	SyncSelectionHandle = GEditor->GetTimerManager()->SetTimerForNextTick(
-	[WeakRootViewModel, Transaction = MoveTemp(Transaction), SelectionChange = MoveTemp(SelectionChange)]() mutable
-	{
-		bSyncSelectionRequested = false;
-	
-		TSharedPtr<FSequencerEditorViewModel> RootViewModel = WeakRootViewModel.Pin();
-		if (!RootViewModel.IsValid())
-		{
-			return;
-		}
-		TSharedPtr<ISequencer> Sequencer = RootViewModel->GetSequencer();
-		if (!Sequencer)
-		{
-			return;
-		}
-
-		
-		if (FCurveEditorExtension* This = RootViewModel->CastDynamic<FCurveEditorExtension>())
-		{
-			This->SyncSelectionHandle.Invalidate();
-			This->WidgetOwner->SyncSelection();
-		}
-
-		// Order matters: this appends a sub-transaction for the selection change...
-		SelectionChange.Reset();
-		// ... and only then can the parent transaction be closed. 
-		Transaction.Reset();
-	});
 }
 
-void FCurveEditorExtension::CleanUpPendingSyncSelection()
-{
-	bSyncSelectionRequested = false;
-	
-	if (GEditor)
-	{
-		GEditor->GetTimerManager()->ClearTimer(SyncSelectionHandle);
-	}
-}
 
 } // namespace Sequencer
 } // namespace UE
diff --git a/Engine/Source/Editor/Sequencer/Private/MVVM/CurveEditorIntegrationExtension.cpp b/Engine/Source/Editor/Sequencer/Private/MVVM/CurveEditorIntegrationExtension.cpp
--- a/Engine/Source/Editor/Sequencer/Private/MVVM/CurveEditorIntegrationExtension.cpp
+++ b/Engine/Source/Editor/Sequencer/Private/MVVM/CurveEditorIntegrationExtension.cpp
@@ -14,31 +14,23 @@
 #include "MVVM/ViewModels/ViewModel.h"
 #include "MVVM/ViewModels/ViewModelIterators.h"
 #include "SequencerSettings.h"
+#include "CurveEditor/SequencerCurveEditorApp.h"
 #include "Views/CurveEditor/CurveEditorWidgetOwner.h"
 
 namespace UE
 {
 namespace Sequencer
 {
+PRAGMA_DISABLE_DEPRECATION_WARNINGS
 
 FCurveEditorIntegrationExtension::FCurveEditorIntegrationExtension()
 {
 }
 
-FCurveEditorIntegrationExtension::~FCurveEditorIntegrationExtension()
-{
-	UnbindFilterBarDependencies();
-}
-
 void FCurveEditorIntegrationExtension::OnCreated(TSharedRef<FViewModel> InWeakOwner)
 {
 	ensureMsgf(!WeakOwnerModel.Pin().IsValid(), TEXT("This extension was already created!"));
 	WeakOwnerModel = InWeakOwner->CastThisShared<FSequenceModel>();
-
-	FSimpleMulticastDelegate& HierarchyChanged = InWeakOwner->GetSharedData()->SubscribeToHierarchyChanged(InWeakOwner);
-	HierarchyChanged.AddSP(this, &FCurveEditorIntegrationExtension::OnHierarchyChanged);
-	
-	SetupFilterBarDependencies();
 }
 
 void FCurveEditorIntegrationExtension::OnHierarchyChanged()
@@ -58,239 +50,29 @@ FCurveEditorExtension* FCurveEditorIntegrationExtension::GetCurveEditorExtension
 	return nullptr;
 }
 
-void FCurveEditorIntegrationExtension::UpdateCurveEditor()
+FSequencerCurveEditorApp* FCurveEditorIntegrationExtension::GetCurveEditorIntegration() const
 {
 	TSharedPtr<FSequenceModel> OwnerModel = WeakOwnerModel.Pin();
-	if (!OwnerModel)
-	{
-		return;
-	}
-	
-	FCurveEditorExtension* CurveEditorExtension = GetCurveEditorExtension();
-	if (!CurveEditorExtension)
-	{
-		return;
-	}
-
-	TSharedPtr<FCurveEditor> CurveEditor = CurveEditorExtension->GetCurveEditor();
-	if (!CurveEditor)
-	{
-		return;
-	}
-	
-	// TODO UE-364008: In the future, we can remove the paths with bUseFilteredOut.
-	const bool bUseFilteredOut = OwnerModel->GetSequencer()->GetSequencerSettings()->GetLinkFiltersWithCurveEditor();
-	const auto IsFilteredOut = [CurveEditorExtension, bUseFilteredOut](const TViewModelPtr<IOutlinerExtension>& InItem)
-	{
-		return UseLinkedFilteringUI()
-			? CurveEditorExtension->GetLinkedFilteringModel()->IsFilteredOut(InItem)
-			: bUseFilteredOut && InItem && InItem->IsFilteredOut();
-	};
-
-	FCurveEditorTree* CurveEditorTree = CurveEditor->GetTree();
-
-	// Guard against multiple broadcasts here and defer them until the end of this function
-	FScopedCurveEditorTreeEventGuard ScopedEventGuard = CurveEditorTree->ScopedEventGuard();
-
-	// Remove any stale tree items
-	for (auto It = ViewModelToTreeItemIDMap.CreateIterator(); It; ++It)
-	{
-		bool bIsRelevant = false;
-		const FCurveEditorTreeItemID ItemID(It.Value());
-		TViewModelPtr<ICurveEditorTreeItemExtension> ViewModel = CastViewModel<ICurveEditorTreeItemExtension>(It.Key().Pin());
-		if (ViewModel)
-		{
-			TViewModelPtr<IOutlinerExtension> OutlinerModel = ViewModel.ImplicitCast();
-			const bool bIsVisible = !IsFilteredOut(OutlinerModel);
-			
-			const FCurveEditorTreeItem* TreeItem = CurveEditorTree->FindItem(ItemID);
-
-			FCurveEditorTreeItemID ParentID;
-			if (auto Parent = ViewModel.AsModel()->FindAncestorOfType<ICurveEditorTreeItemExtension>())
-			{
-				ParentID = ViewModelToTreeItemIDMap.FindRef(Parent.AsModel());
-			}
-			
-			bIsRelevant = bIsVisible && TreeItem && TreeItem->GetParentID() == ParentID;
-		}
-		
-		if (!bIsRelevant)
-		{
-			if (ViewModel)
-			{
-				ViewModel->OnRemovedFromCurveEditor(CurveEditor);
-			}
-			CurveEditor->RemoveTreeItem(ItemID);
-			It.RemoveCurrent();
-		}
-	}
-
-	// Do a second pass to remove any items that were removed recursively above
-	for (auto It = ViewModelToTreeItemIDMap.CreateIterator(); It; ++It)
-	{
-		if (CurveEditorTree->FindItem(It->Value) == nullptr)
-		{
-			TreeItemIDToViewModelMap.Remove(It->Value);
-			It.RemoveCurrent();
-		}
-	}
-
-	// Iterate all non-filtered out outliners and check for curve editor tree extensions
-	const bool bIncludeRootNode = false;
-	bool bItemsAdded = false;
-	for (TParentFirstChildIterator<IOutlinerExtension> It(OwnerModel, bIncludeRootNode); It; ++It)
-	{
-		if (IsFilteredOut(*It))
-		{
-			It.IgnoreCurrentChildren();
-			continue;
-		}
-
-		if (TViewModelPtr<ICurveEditorTreeItemExtension> ChildViewModel = (*It).ImplicitCast())
-		{
-			const FCurveEditorTreeItemID ItemID = ViewModelToTreeItemIDMap.FindRef(ChildViewModel.AsModel());
-			if (!ItemID.IsValid())
-			{
-				bItemsAdded = true;
-				AddToCurveEditor(ChildViewModel, CurveEditor);
-			}
-		}
-	}
-
-	// If new items have been added, synchronize selection after the view models have been added to the curve editor. Synchronization depends on curve model tree item IDs
-	if (bItemsAdded)
-	{
-		CurveEditorExtension->RequestSyncSelection();
-	}
+	const TSharedPtr<ISequencer> Sequencer = OwnerModel ? OwnerModel->GetSequencer() : nullptr;
+	return Sequencer ? FSequencerCurveEditorApp::Get(*Sequencer) : nullptr;
 }
 
-FCurveEditorTreeItemID FCurveEditorIntegrationExtension::AddToCurveEditor(TViewModelPtr<ICurveEditorTreeItemExtension> InViewModel, TSharedPtr<FCurveEditor> InCurveEditor)
+void FCurveEditorIntegrationExtension::UpdateCurveEditor()
 {
-	// If the view model doesn't want to be in the curve editor, bail out
-	// Note that this means we will create curve editor items for each parent in the hierarchy up
-	// until the first parent that doesn't implement ICurveEditorTreeItemExtension
-	// That is: we don't create "dummy" entries when there's a "gap" in the hierarchy
-	if (!InViewModel)
+	if (FSequencerCurveEditorApp* CurveEditorApp = GetCurveEditorIntegration())
 	{
-		return FCurveEditorTreeItemID::Invalid();
+		CurveEditorApp->UpdateCurveEditor();
 	}
-
-	// Check if we already have a valid curve editor ID
-	if (FCurveEditorTreeItemID* Existing = ViewModelToTreeItemIDMap.Find(InViewModel.AsModel()))
-	{
-		if (InCurveEditor->GetTree()->FindItem(*Existing) != nullptr)
-		{
-			return *Existing;
-		}
-	}
-
-	// Recursively create any needed parent curve editor items
-	TViewModelPtr<ICurveEditorTreeItemExtension> Parent = InViewModel.AsModel()->CastParent<ICurveEditorTreeItemExtension>();
-
-	FCurveEditorTreeItemID ParentID = Parent ?
-		AddToCurveEditor(Parent, InCurveEditor) :
-		FCurveEditorTreeItemID::Invalid();
-
-	// Create the new curve editor item
-	FCurveEditorTreeItem* NewItem = InCurveEditor->AddTreeItem(ParentID);
-	TSharedPtr<ICurveEditorTreeItem> CurveEditorTreeItem = InViewModel->GetCurveEditorTreeItem();
-	NewItem->SetWeakItem(CurveEditorTreeItem);
-	TOptional<FString> UniquePathName =  InViewModel->GetUniquePathName();
-	NewItem->SetUniquePathName(UniquePathName);
-	
-	// Register the new ID in our map and notify the view model
-	ViewModelToTreeItemIDMap.Add(InViewModel.AsModel(), NewItem->GetID());
-	TreeItemIDToViewModelMap.Add(NewItem->GetID(), InViewModel.AsModel());
-	
-	InViewModel->OnAddedToCurveEditor(NewItem->GetID(), InCurveEditor);
-
-	return NewItem->GetID();
 }
 
 void FCurveEditorIntegrationExtension::ResetCurveEditor()
 {
-	FCurveEditorExtension* CurveEditorExtension = GetCurveEditorExtension();
-	if (!ensure(CurveEditorExtension))
-	{
-		return;
-	}
-
-	TSharedPtr<FCurveEditor> CurveEditor = CurveEditorExtension->GetCurveEditor();
-	if (!ensure(CurveEditor))
-	{
-		return;
-	}
-
-	for (TPair<TWeakPtr<FViewModel>, FCurveEditorTreeItemID> Pair : ViewModelToTreeItemIDMap)
-	{
-		FCurveEditorTreeItemID ItemID(Pair.Value);
-		TViewModelPtr<ICurveEditorTreeItemExtension> ViewModel = CastViewModel<ICurveEditorTreeItemExtension>(Pair.Key.Pin());
-		if (ViewModel)
-		{
-			ViewModel->OnRemovedFromCurveEditor(CurveEditor);
-		}
-		CurveEditor->RemoveTreeItem(ItemID);
-	}
-
-	ViewModelToTreeItemIDMap.Reset();
-	TreeItemIDToViewModelMap.Reset();
-}
-
-TWeakPtr<FViewModel> FCurveEditorIntegrationExtension::FindViewModelById(const FCurveEditorTreeItemID& InTreeId) const
-{
-	const TWeakPtr<FViewModel>* ViewModel = TreeItemIDToViewModelMap.Find(InTreeId);
-	return ViewModel ? *ViewModel : nullptr;
-}
-
-FCurveEditorTreeItemID FCurveEditorIntegrationExtension::FindIdByViewModel(TWeakPtr<FViewModel> InViewModel) const
-{
-	const FCurveEditorTreeItemID* TreeItemID = ViewModelToTreeItemIDMap.Find(InViewModel);
-	return TreeItemID ? *TreeItemID : FCurveEditorTreeItemID::Invalid();
-}
-
-void FCurveEditorIntegrationExtension::SetupFilterBarDependencies()
-{
-	// TODO UE-364009: Effectively, this makes FCurveEditorIntegrationExtension bind to FilterBarProxy. In the future, refactor 
-	// FCurveEditorIntegrationExtension to be a separate class that FCurveEditorExtension injects the dependencies to.
-	FCurveEditorExtension* CurveEditorExtension = GetCurveEditorExtension();
-	if (!ensure(CurveEditorExtension))
-	{
-		return;
-	}
-	
-	if (CurveEditorExtension->GetLinkedFilteringModel())
-	{
-		BindFilterBarDependencies();
-	}
-	else
-	{
-		CurveEditorExtension->OnCurveEditorCreated.AddRaw(this, &FCurveEditorIntegrationExtension::BindFilterBarDependencies);
-	}
-}
-
-void FCurveEditorIntegrationExtension::BindFilterBarDependencies()
-{
-	if (FCurveEditorExtension* CurveEditorExtension = GetCurveEditorExtension())
-	{
-		CurveEditorExtension->GetLinkedFilteringModel()->OnFilteredItemsChanged().AddRaw(this, &FCurveEditorIntegrationExtension::UpdateCurveEditor);
-		UpdateCurveEditor();
-	}
-}
-
-void FCurveEditorIntegrationExtension::UnbindFilterBarDependencies()
-{
-	FCurveEditorExtension* CurveEditorExtension = GetCurveEditorExtension();
-	if (!CurveEditorExtension)
-	{
-		return;
-	}
-	
-	CurveEditorExtension->OnCurveEditorCreated.RemoveAll(this);
-	if (ILinkedFilterViewModel* ViewModel = CurveEditorExtension->GetLinkedFilteringModel())
+	if (FSequencerCurveEditorApp* CurveEditorApp = GetCurveEditorIntegration())
 	{
-		ViewModel->OnFilteredItemsChanged().RemoveAll(this);
+		CurveEditorApp->ResetCurveEditor();
 	}
 }
+PRAGMA_ENABLE_DEPRECATION_WARNINGS
 } // namespace Sequencer
 } // namespace UE
 
diff --git a/Engine/Source/Editor/Sequencer/Private/MVVM/SequencerViewModelTypeIDs.cpp b/Engine/Source/Editor/Sequencer/Private/MVVM/SequencerViewModelTypeIDs.cpp
--- a/Engine/Source/Editor/Sequencer/Private/MVVM/SequencerViewModelTypeIDs.cpp
+++ b/Engine/Source/Editor/Sequencer/Private/MVVM/SequencerViewModelTypeIDs.cpp
@@ -81,7 +81,9 @@ UE_SEQUENCER_DEFINE_CASTABLE(FSequencerTrackAreaViewModel);
 
 // Extension types
 UE_SEQUENCER_DEFINE_VIEW_MODEL_TYPE_ID(FCurveEditorExtension);
+PRAGMA_DISABLE_DEPRECATION_WARNINGS
 UE_SEQUENCER_DEFINE_VIEW_MODEL_TYPE_ID(FCurveEditorIntegrationExtension);
+PRAGMA_ENABLE_DEPRECATION_WARNINGS
 UE_SEQUENCER_DEFINE_VIEW_MODEL_TYPE_ID(FFolderModelStorageExtension);
 UE_SEQUENCER_DEFINE_VIEW_MODEL_TYPE_ID(FObjectBindingModelStorageExtension);
 UE_SEQUENCER_DEFINE_VIEW_MODEL_TYPE_ID(FPinEditorExtension);
diff --git a/Engine/Source/Editor/Sequencer/Private/MVVM/Views/CurveEditor/CurveEditorWidgetOwner.cpp b/Engine/Source/Editor/Sequencer/Private/MVVM/Views/CurveEditor/CurveEditorWidgetOwner.cpp
--- a/Engine/Source/Editor/Sequencer/Private/MVVM/Views/CurveEditor/CurveEditorWidgetOwner.cpp
+++ b/Engine/Source/Editor/Sequencer/Private/MVVM/Views/CurveEditor/CurveEditorWidgetOwner.cpp
@@ -37,6 +37,7 @@ static FCurveEditorWidgets MakePresenterWidgets(
 	const TSharedRef<FCurveEditor>& InCurveEditor, 
 	const TSharedRef<ITimeSliderController>& InTimeSliderController,
 	const TSharedRef<FLinkedFilterViewModel>& InFilteringViewModel,
+	const FCurveModelSyncer& InCurveModelSyncer UE_LIFETIMEBOUND,
 	const TSharedPtr<FTabManager> InTabManager
 	)
 {
@@ -47,6 +48,7 @@ static FCurveEditorWidgets MakePresenterWidgets(
 		.Sequencer(InSequencer)
 		.CurveEditor(InCurveEditor)
 		.FilterViewModel(InFilteringViewModel)
+		.CurveModelSyncer(&InCurveModelSyncer)
 		.FilterPills()
 		[
 			SNew(SFilterBarSwitcher)
@@ -86,11 +88,14 @@ FCurveEditorWidgetOwner::FCurveEditorWidgetOwner(
 	const TSharedRef<FSequencer>& InSequencer,
 	const TSharedRef<FCurveEditor>& InCurveEditor,
 	const TSharedRef<ITimeSliderController>& InTimeSliderController,
-	const TSharedRef<FLinkedFilterViewModel>& InFilteringViewModel
+	const TSharedRef<FLinkedFilterViewModel>& InFilteringViewModel,
+	const FCurveModelSyncer& InCurveModelSyncer
 	)
 	: WeakSequencer(InSequencer)
 	, CurveEditor(InCurveEditor)
-	, Widgets(PresenterDetail::MakePresenterWidgets(InSequencer, InCurveEditor, InTimeSliderController, InFilteringViewModel, GetTabManager()))
+	, Widgets(PresenterDetail::MakePresenterWidgets(
+		InSequencer, InCurveEditor, InTimeSliderController, InFilteringViewModel, InCurveModelSyncer, GetTabManager()
+		))
 {
 	BindCommands();
 	InitFrameNumberPropertyLayout(InSequencer);
diff --git a/Engine/Source/Editor/Sequencer/Private/MVVM/Views/CurveEditor/CurveEditorWidgetOwner.h b/Engine/Source/Editor/Sequencer/Private/MVVM/Views/CurveEditor/CurveEditorWidgetOwner.h
--- a/Engine/Source/Editor/Sequencer/Private/MVVM/Views/CurveEditor/CurveEditorWidgetOwner.h
+++ b/Engine/Source/Editor/Sequencer/Private/MVVM/Views/CurveEditor/CurveEditorWidgetOwner.h
@@ -3,6 +3,7 @@
 #pragma once
 
 #include "CurveEditorWidgets.h"
+#include "HAL/Platform.h"
 #include "Templates/SharedPointer.h"
 
 class FCurveEditor;
@@ -17,6 +18,7 @@ struct FTimeSliderArgs;
 
 namespace UE::Sequencer
 {
+class FCurveModelSyncer;
 class FLinkedFilterViewModel;
 
 /** @return Whether the linked filtering UI is enabled. */
@@ -38,7 +40,8 @@ public:
 		const TSharedRef<FSequencer>& InSequencer, 
 		const TSharedRef<FCurveEditor>& InCurveEditor,
 		const TSharedRef<ITimeSliderController>& InTimeSliderController,
-		const TSharedRef<FLinkedFilterViewModel>& InFilteringViewModel
+		const TSharedRef<FLinkedFilterViewModel>& InFilteringViewModel,
+		const FCurveModelSyncer& InCurveModelSyncer UE_LIFETIMEBOUND
 		);
 	~FCurveEditorWidgetOwner();
 	
diff --git a/Engine/Source/Editor/Sequencer/Private/MVVM/Views/CurveEditor/SCurveTreeContent.cpp b/Engine/Source/Editor/Sequencer/Private/MVVM/Views/CurveEditor/SCurveTreeContent.cpp
--- a/Engine/Source/Editor/Sequencer/Private/MVVM/Views/CurveEditor/SCurveTreeContent.cpp
+++ b/Engine/Source/Editor/Sequencer/Private/MVVM/Views/CurveEditor/SCurveTreeContent.cpp
@@ -10,6 +10,7 @@
 #include "SPopoutTabInlineContent.h"
 #include "STemporarilyFocusedSpinBox.h"
 #include "Sequencer.h"
+#include "CurveEditor/Integration/CurveModelSyncer.h"
 #include "Filters/Widgets/Linking/SFilterBarSwitcher.h"
 #include "MVVM/CurveEditorIntegrationExtension.h"
 #include "Styling/AppStyle.h"
@@ -19,14 +20,14 @@
 
 namespace UE::Sequencer
 {
-class FCurveEditorIntegrationExtension;
-
 void SCurveTreeContent::Construct(const FArguments& InArgs)
 {
-	check(InArgs._Sequencer && InArgs._CurveEditor && InArgs._FilterViewModel);
+	const bool bHasCurveModelSyncer = InArgs._CurveModelSyncer.IsBound() || InArgs._CurveModelSyncer.IsSet();
+	check(InArgs._Sequencer && InArgs._CurveEditor && InArgs._FilterViewModel && bHasCurveModelSyncer);
 	
 	WeakSequencer = InArgs._Sequencer;
 	CurveEditor = InArgs._CurveEditor;
+	CurveModelSyncerAttr = InArgs._CurveModelSyncer;
 	
 	ChildSlot
 	[
@@ -159,16 +160,13 @@ TSharedRef<SWidget> SCurveTreeContent::MakeFooter(const FArguments& InArgs)
 
 void SCurveTreeContent::ScrollIntoView(const TWeakViewModelPtr<IOutlinerExtension>& WeakViewModel)
 {
-	// TODO UE-364009: Change this to use the new class that FCurveEditorIntegrationExtension is extracted to (and inject it to SCurveTreeContent)
-	const TSharedPtr<FSequencer> Sequencer = WeakSequencer.Pin();
-	FCurveEditorIntegrationExtension* CurveEditorIntegration = Sequencer 
-		? Sequencer->GetViewModel()->GetRootModel()->CastDynamic<FCurveEditorIntegrationExtension>() : nullptr;
-	if (!CurveEditorIntegration)
+	const FCurveModelSyncer* CurveModelSyncer = CurveModelSyncerAttr.Get();
+	if (!CurveModelSyncer)
 	{
 		return;
 	}
 	
-	const FCurveEditorTreeItemID Id = CurveEditorIntegration->FindIdByViewModel(WeakViewModel);
+	const FCurveEditorTreeItemID Id = CurveModelSyncer->FindIdByViewModel(WeakViewModel);
 	if (Id.IsValid())
 	{
 		CurveEditorTreeView->RequestScrollIntoView(Id);
@@ -177,11 +175,8 @@ void SCurveTreeContent::ScrollIntoView(const TWeakViewModelPtr<IOutlinerExtensio
 
 TSet<TWeakViewModelPtr<IOutlinerExtension>> SCurveTreeContent::GetAllItems() const
 {
-	// TODO UE-364009: Change this to use the new class that FCurveEditorIntegrationExtension is extracted to (and inject it to SCurveTreeContent)
-	const TSharedPtr<FSequencer> Sequencer = WeakSequencer.Pin();
-	FCurveEditorIntegrationExtension* CurveEditorIntegration = Sequencer 
-		? Sequencer->GetViewModel()->GetRootModel()->CastDynamic<FCurveEditorIntegrationExtension>() : nullptr;
-	if (!CurveEditorIntegration)
+	const FCurveModelSyncer* CurveModelSyncer = CurveModelSyncerAttr.Get();
+	if (!CurveModelSyncer)
 	{
 		return {};
 	}
@@ -189,7 +184,7 @@ TSet<TWeakViewModelPtr<IOutlinerExtension>> SCurveTreeContent::GetAllItems() con
 	TSet<TWeakViewModelPtr<IOutlinerExtension>> AllItems;
 	for (const TPair<FCurveEditorTreeItemID, FCurveEditorTreeItem>& Item : CurveEditor->GetTree()->GetAllItems())
 	{
-		const TSharedPtr<FViewModel> ViewModel = CurveEditorIntegration->FindViewModelById(Item.Key).Pin();
+		const TSharedPtr<FViewModel> ViewModel = CurveModelSyncer->FindViewModelById(Item.Key).Pin();
 		const TViewModelPtr<IOutlinerExtension> OutlinerNode = ViewModel ? CastViewModel<IOutlinerExtension>(ViewModel.ToSharedRef()) : nullptr;
 		if (OutlinerNode)
 		{
@@ -201,11 +196,8 @@ TSet<TWeakViewModelPtr<IOutlinerExtension>> SCurveTreeContent::GetAllItems() con
 
 TSet<TWeakViewModelPtr<IOutlinerExtension>> SCurveTreeContent::GetSelectedItems() const
 {
-	// TODO UE-364009: Change this to use the new class that FCurveEditorIntegrationExtension is extracted to (and inject it to SCurveTreeContent)
-	const TSharedPtr<FSequencer> Sequencer = WeakSequencer.Pin();
-	FCurveEditorIntegrationExtension* CurveEditorIntegration = Sequencer 
-		? Sequencer->GetViewModel()->GetRootModel()->CastDynamic<FCurveEditorIntegrationExtension>() : nullptr;
-	if (!CurveEditorIntegration)
+	const FCurveModelSyncer* CurveModelSyncer = CurveModelSyncerAttr.Get();
+	if (!CurveModelSyncer)
 	{
 		return {};
 	}
@@ -213,7 +205,7 @@ TSet<TWeakViewModelPtr<IOutlinerExtension>> SCurveTreeContent::GetSelectedItems(
 	TSet<TWeakViewModelPtr<IOutlinerExtension>> Selected;
 	for (const FCurveEditorTreeItemID& Item : CurveEditorTreeView->GetSelectedItems())
 	{
-		const TSharedPtr<FViewModel> ViewModel = CurveEditorIntegration->FindViewModelById(Item).Pin();
+		const TSharedPtr<FViewModel> ViewModel = CurveModelSyncer->FindViewModelById(Item).Pin();
 		const TViewModelPtr<IOutlinerExtension> OutlinerNode = ViewModel ? CastViewModel<IOutlinerExtension>(ViewModel.ToSharedRef()) : nullptr;
 		if (OutlinerNode)
 		{
diff --git a/Engine/Source/Editor/Sequencer/Private/MVVM/Views/CurveEditor/SCurveTreeContent.h b/Engine/Source/Editor/Sequencer/Private/MVVM/Views/CurveEditor/SCurveTreeContent.h
--- a/Engine/Source/Editor/Sequencer/Private/MVVM/Views/CurveEditor/SCurveTreeContent.h
+++ b/Engine/Source/Editor/Sequencer/Private/MVVM/Views/CurveEditor/SCurveTreeContent.h
@@ -14,6 +14,7 @@ template<typename T> struct STemporarilyFocusedSpinBox;
 
 namespace UE::Sequencer
 {
+class FCurveModelSyncer;
 class FLinkedFilterViewModel;
 
 /** Displays Curve Editor in Sequencer with linked filtering. */
@@ -32,6 +33,9 @@ public:
 		/** Explicit tab manager to use, otherwise uses the global tab manager */
 		SLATE_ARGUMENT(TSharedPtr<FTabManager>, TabManager)
 
+		/** Used to translate between IOutlinerExtension and FCurveEditorTreeItemID. */
+		SLATE_ATTRIBUTE(const FCurveModelSyncer*, CurveModelSyncer)
+
 		/** Shows the filter pills. */
 		SLATE_NAMED_SLOT(FArguments, FilterPills)
 	SLATE_END_ARGS()
@@ -55,6 +59,9 @@ private:
 	/** The curve editor being displayed. */
 	TSharedPtr<FCurveEditor> CurveEditor;
 	
+	/** Used to translate between IOutlinerExtension and FCurveEditorTreeItemID. */
+	TAttribute<const FCurveModelSyncer*> CurveModelSyncerAttr;
+	
 	/** Curve editor tree widget */
 	TSharedPtr<SCurveEditorTree> CurveEditorTreeView;
 	/** The current playback time display. */
diff --git a/Engine/Source/Editor/Sequencer/Private/Sequencer.cpp b/Engine/Source/Editor/Sequencer/Private/Sequencer.cpp
--- a/Engine/Source/Editor/Sequencer/Private/Sequencer.cpp
+++ b/Engine/Source/Editor/Sequencer/Private/Sequencer.cpp
@@ -189,6 +189,7 @@
 #include "Bindings/MovieSceneSpawnableBinding.h"
 #include "Bindings/MovieSceneSpawnableActorBinding.h"
 #include "Bindings/MovieSceneReplaceableActorBinding.h"
+#include "CurveEditor/SequencerCurveEditorApp.h"
 #include "Decorations/MovieSceneMuteSoloDecoration.h"
 #include "Filters/SequencerTrackFilterCommands.h"
 #include "Filters/ViewModels/Linking/LinkedFilterViewModel.h"
@@ -2894,17 +2895,13 @@ void FSequencer::RecreateCurveEditor()
 {
 	using namespace UE::Sequencer;
 
-	FCurveEditorIntegrationExtension* CurveEditorIntegration= ViewModel->GetRootModel()->CastDynamic<FCurveEditorIntegrationExtension>();
+	FSequencerCurveEditorApp* CurveEditorIntegration = FSequencerCurveEditorApp::Get(*this);
 	if (ensure(CurveEditorIntegration))
 	{
 		CurveEditorIntegration->ResetCurveEditor();
 		if (GetSequencerSettings()->ShouldSyncCurveEditorSelection())
 		{
-			FCurveEditorExtension* CurveEditorExtension = ViewModel->CastDynamic<FCurveEditorExtension>();
-			if (ensure(CurveEditorExtension))
-			{
-				return CurveEditorExtension->RequestSyncSelection();
-			}
+			CurveEditorIntegration->RequestSyncSelection();
 		}
 	}
 }
diff --git a/Engine/Source/Editor/Sequencer/Private/SequencerModule.cpp b/Engine/Source/Editor/Sequencer/Private/SequencerModule.cpp
--- a/Engine/Source/Editor/Sequencer/Private/SequencerModule.cpp
+++ b/Engine/Source/Editor/Sequencer/Private/SequencerModule.cpp
@@ -531,13 +531,6 @@ public:
 				InModel->AddDynamicExtension(FTrackRowModelStorageExtension::ID);
 				InModel->AddDynamicExtension(FSectionModelStorageExtension::ID);
 				InModel->AddDynamicExtension(FDecorationModelStorageExtension::ID);
-
-				// If the editor supports a curve editor, add an integration extension to
-				// sync view-model hierarchies between the outliner and curve editor.
-				if (InEditor->CastDynamic<FCurveEditorExtension>())
-				{
-					InModel->AddDynamicExtension(FCurveEditorIntegrationExtension::ID);
-				}
 			}
 		);
 
diff --git a/Engine/Source/Editor/Sequencer/Private/SequencerNodeTree.cpp b/Engine/Source/Editor/Sequencer/Private/SequencerNodeTree.cpp
--- a/Engine/Source/Editor/Sequencer/Private/SequencerNodeTree.cpp
+++ b/Engine/Source/Editor/Sequencer/Private/SequencerNodeTree.cpp
@@ -110,13 +110,7 @@ void FSequencerNodeTree::Update()
 		IGeometryExtension::UpdateVirtualGeometry(0.f, RootNode);
 	}
 
-	// Update curve editor tree based on new filtered hierarchy
-	auto CurveEditorIntegration = SequenceModel->CastDynamic<FCurveEditorIntegrationExtension>();
-	if (CurveEditorIntegration)
-	{
-		CurveEditorIntegration->UpdateCurveEditor();
-	}
-
+	// This will cause relevant systems to update. E.g. curve editor, if enabled, will updated based on the new filtered hierarchy.
 	OnUpdatedDelegate.Broadcast();
 }
 
diff --git a/Engine/Source/Editor/Sequencer/Public/MVVM/CurveEditorExtension.h b/Engine/Source/Editor/Sequencer/Public/MVVM/CurveEditorExtension.h
--- a/Engine/Source/Editor/Sequencer/Public/MVVM/CurveEditorExtension.h
+++ b/Engine/Source/Editor/Sequencer/Public/MVVM/CurveEditorExtension.h
@@ -27,10 +27,10 @@ namespace UE
 {
 namespace Sequencer
 {
-class FCurveEditorWidgetOwner;
 class FSequencerEditorViewModel;
-class ILinkedFilterViewModel;
+class FSequencerCurveEditorApp;
 
+/** Public API for interacting with Curve Editor in Sequencer. */
 class FCurveEditorExtension : public IDynamicExtension
 {
 public:
@@ -46,7 +46,7 @@ public:
 	UE_API void CreateCurveEditor(const FTimeSliderArgs& TimeSliderArgs);
 
 	/** Gets the curve editor view-model */
-	TSharedPtr<FCurveEditor> GetCurveEditor() const { return CurveEditorModel; }
+	UE_API TSharedPtr<FCurveEditor> GetCurveEditor() const;
 
 	/** Opens the curve editor */
 	UE_API void OpenCurveEditor();
@@ -68,46 +68,18 @@ public:
 	static UE_API const FName CurveEditorTabName;
 
 private:
-
+	
 	TSharedPtr<ISequencer> GetSequencer() const;
 	USequencerSettings* GetSequencerSettings() const;
 	TSharedPtr<FTabManager> GetTabManager() const;
 
-	/** Get the default key attributes to apply to newly created keys on the curve editor */
-	UE_API FKeyAttributes GetDefaultKeyAttributes() const;
-	
-	/** Cancels any pending transactions. */
-	void CleanUpPendingSyncSelection();
-	
-	/** @return The proxy for interacting with linked filter */
-	ILinkedFilterViewModel* GetLinkedFilteringModel() const { return CurveEditorFilterClient.Get(); }
-
-private:
-
 	/** The sequencer editor we are extending with a curve editor */
 	TWeakPtr<FSequencerEditorViewModel> WeakOwnerModel;
-
-	/** Curve editor */
-	TSharedPtr<FCurveEditor> CurveEditorModel;
-	
-	/** Handles refiltering when the active filter bar becomes dirty. */
-	TSharedPtr<ILinkedFilterViewModel> CurveEditorFilterClient;
-	
-	/** Responsible for creating the UI. */
-	TPimplPtr<FCurveEditorWidgetOwner> WidgetOwner; 
-
-	/** Handle to latent refresh of selection. */
-	FTimerHandle SyncSelectionHandle;
 	
-	/**
-	 * Let's FCurveEditorIntegrationExtension know that the curve editor is set up.
-	 * 
-	 * TODO UE-364009: Instead, FCurveEditorIntegrationExtension should be extracted to a separate class that is created by FCurveEditorExtension 
-	 * so any dependencies can be injected.
-	 */
-	FSimpleMulticastDelegate OnCurveEditorCreated;
+	/** Manages the Curve Editor x Sequencer integration. */
+	TPimplPtr<FSequencerCurveEditorApp> CurveEditorApp;
 	
-	friend class FCurveEditorIntegrationExtension;
+	friend class FSequencerCurveEditorApp; // Engineering compromise to access FSequencerCurveEditorApp without exposing to public API
 };
 
 } // namespace Sequencer
diff --git a/Engine/Source/Editor/Sequencer/Public/MVVM/CurveEditorIntegrationExtension.h b/Engine/Source/Editor/Sequencer/Public/MVVM/CurveEditorIntegrationExtension.h
--- a/Engine/Source/Editor/Sequencer/Public/MVVM/CurveEditorIntegrationExtension.h
+++ b/Engine/Source/Editor/Sequencer/Public/MVVM/CurveEditorIntegrationExtension.h
@@ -30,14 +30,16 @@ class FCurveEditorExtension;
  * - Outliner items implementing ICurveEditorTreeItemExtension (or its default shim) if they 
  *   want to show up in the curve editor.
  */
-class FCurveEditorIntegrationExtension : public IDynamicExtension
+class UE_DEPRECATED(5.8, "This class has been removed.") FCurveEditorIntegrationExtension : public IDynamicExtension
 {
 public:
 
+	PRAGMA_DISABLE_DEPRECATION_WARNINGS
 	UE_SEQUENCER_DECLARE_VIEW_MODEL_TYPE_ID_API(UE_API, FCurveEditorIntegrationExtension)
+	PRAGMA_ENABLE_DEPRECATION_WARNINGS
 
+	UE_DEPRECATED(5.8, "This class has been removed.")
 	UE_API FCurveEditorIntegrationExtension();
-	UE_API virtual ~FCurveEditorIntegrationExtension() override;
 
 	UE_API virtual void OnCreated(TSharedRef<FViewModel> InWeakOwner) override;
 
@@ -45,41 +47,28 @@ public:
 	 * Keeps the curve editor items up-to-date with the sequencer outliner by adding/removing 
 	 * entries as needed.
 	 */
+	UE_DEPRECATED(5.8, "This function has been removed.")
 	UE_API void UpdateCurveEditor();
 
 	/**
 	 * Clears the curve editor of all contents.
 	 */
+	UE_DEPRECATED(5.8, "This function has been removed.")
 	UE_API void ResetCurveEditor();
-	
-	/** @return The view model associated with the tree item. */
-	TWeakPtr<FViewModel> FindViewModelById(const FCurveEditorTreeItemID& InTreeId) const;
-	/** @return The tree item associated with the view model. */
-	FCurveEditorTreeItemID FindIdByViewModel(TWeakPtr<FViewModel>) const;
 
 private:
 
 	/** Update curve editor items when sequencer outliner items change */
 	UE_API void OnHierarchyChanged();
 
-	/** Adds the given view-model to the curve editor */
-	UE_API FCurveEditorTreeItemID AddToCurveEditor(TViewModelPtr<ICurveEditorTreeItemExtension> InViewModel, TSharedPtr<FCurveEditor> InCurveEditor);
-
+	class FSequencerCurveEditorApp* GetCurveEditorIntegration() const;
+	
 	/** Finds the curve editor extension on the top-level sequencer editor view-model */
 	UE_API FCurveEditorExtension* GetCurveEditorExtension();
 
 private:
 
 	TWeakPtr<FSequenceModel> WeakOwnerModel;
-	TMap<TWeakPtr<FViewModel>, FCurveEditorTreeItemID> ViewModelToTreeItemIDMap;
-	TMap<FCurveEditorTreeItemID, TWeakPtr<FViewModel>> TreeItemIDToViewModelMap;
-	
-	/** Ensures BindFilterBarDependencies is called at some point. */
-	void SetupFilterBarDependencies();
-	/** Binds to the linked filtering system to cause refreshing of curve editor content when the filtering state changes. */
-	void BindFilterBarDependencies();
-	/** Cleans up any callbacks to the linked filtering system. */
-	void UnbindFilterBarDependencies();
 };
 
 } // namespace Sequencer
