Sequencer: Fix unset pointers in SSequencerSearchAndFilterRow. This fixes a crash when you focus the search box using "Ctrl+F".

#Synced-CL 51334859

diff --git a/Engine/Source/Editor/Sequencer/Private/Widgets/Sequencer/SSequencerSearchAndFilterRow.cpp b/Engine/Source/Editor/Sequencer/Private/Widgets/Sequencer/SSequencerSearchAndFilterRow.cpp
--- a/Engine/Source/Editor/Sequencer/Private/Widgets/Sequencer/SSequencerSearchAndFilterRow.cpp
+++ b/Engine/Source/Editor/Sequencer/Private/Widgets/Sequencer/SSequencerSearchAndFilterRow.cpp
@@ -146,15 +146,17 @@ TSharedRef<SWidget> SSequencerSearchAndFilterRow::MakeFilterContent(const FArgum
 			];
 	};
 	
+	LinkedFilterWidgets = MakeWidget(LinkedFilterViewModel->GetLinkedFilterBarImpl());
+	InstancedFilterWidgets = MakeWidget(LinkedFilterViewModel->GetInstancedFilterBarImpl());
 	return SNew(SFilterLinkStateSwitcher)
 		.LinkedFilterViewModel(InArgs._FilterViewModel)
 		.LinkedContent() 
 		[
-			MakeWidget(LinkedFilterViewModel->GetLinkedFilterBarImpl())
+			LinkedFilterWidgets.ToSharedRef()
 		]
 		.InstancedContent() 
 		[
-			MakeWidget(LinkedFilterViewModel->GetInstancedFilterBarImpl())
+			InstancedFilterWidgets.ToSharedRef()
 		];
 }
 
