Sequencer: Implement "Ctrl+F" for search box in linked filtering UI.

#Synced-CL 51334859

diff --git a/Engine/Source/Editor/Sequencer/Private/MVVM/Views/CurveEditor/SCurveEditorSearchAndFilterRow.cpp b/Engine/Source/Editor/Sequencer/Private/MVVM/Views/CurveEditor/SCurveEditorSearchAndFilterRow.cpp
--- a/Engine/Source/Editor/Sequencer/Private/MVVM/Views/CurveEditor/SCurveEditorSearchAndFilterRow.cpp
+++ b/Engine/Source/Editor/Sequencer/Private/MVVM/Views/CurveEditor/SCurveEditorSearchAndFilterRow.cpp
@@ -44,6 +44,17 @@ void SCurveEditorSearchAndFilterRow::Construct(const FArguments& InArgs)
 	];
 }
 
+TSharedPtr<SSequencerSearchBox> SCurveEditorSearchAndFilterRow::GetActiveSearchBox() const
+{
+	static_assert(static_cast<int32>(ELinkedFilterMode::Count) == 2, "Update this switch");
+	switch (FilterModel->GetFilterMode())
+	{
+	case ELinkedFilterMode::Linked: return LinkedFilterWidgets->GetSearchBox();
+	case ELinkedFilterMode::Instanced: return InstancedFilterWidgets->GetSearchBox();
+	default: checkNoEntry(); return nullptr;
+	}
+}
+
 TSharedRef<SWidget> SCurveEditorSearchAndFilterRow::MakeFilterContent(const FArguments& InArgs)
 {
 	const auto MakeWidget = [this, &InArgs](const TSharedRef<FSequencerFilterBar>& InFilterBar)
@@ -70,10 +81,18 @@ TSharedRef<SWidget> SCurveEditorSearchAndFilterRow::MakeFilterContent(const FArg
 	};
 	
 	const TSharedPtr<FLinkedFilterViewModel> LinkedFilterViewModel = InArgs._FilterViewModel;
+	LinkedFilterWidgets = MakeWidget(LinkedFilterViewModel->GetLinkedFilterBarImpl());
+	InstancedFilterWidgets = MakeWidget(LinkedFilterViewModel->GetInstancedFilterBarImpl());
 	return SNew(SFilterLinkStateSwitcher)
 		.LinkedFilterViewModel(InArgs._FilterViewModel)
-		.LinkedContent() [ MakeWidget(LinkedFilterViewModel->GetLinkedFilterBarImpl()) ]
-		.InstancedContent() [ MakeWidget(LinkedFilterViewModel->GetInstancedFilterBarImpl()) ];
+		.LinkedContent() 
+		[
+			LinkedFilterWidgets.ToSharedRef()
+		]
+		.InstancedContent() 
+		[
+			InstancedFilterWidgets.ToSharedRef()
+		];
 }
 
 void SCurveEditorSearchAndFilterRow::OnSearchTextChanged(const FText& InNewText, TSharedPtr<ISequencerTrackFilters> InFilter)
diff --git a/Engine/Source/Editor/Sequencer/Private/MVVM/Views/CurveEditor/SCurveEditorSearchAndFilterRow.h b/Engine/Source/Editor/Sequencer/Private/MVVM/Views/CurveEditor/SCurveEditorSearchAndFilterRow.h
--- a/Engine/Source/Editor/Sequencer/Private/MVVM/Views/CurveEditor/SCurveEditorSearchAndFilterRow.h
+++ b/Engine/Source/Editor/Sequencer/Private/MVVM/Views/CurveEditor/SCurveEditorSearchAndFilterRow.h
@@ -6,6 +6,7 @@
 #include "Widgets/DeclarativeSyntaxSupport.h"
 #include "Widgets/SCompoundWidget.h"
 
+class SSequencerSearchBox;
 class ISequencerTrackFilters;
 class FSequencer;
 
@@ -37,6 +38,9 @@ public:
 	SLATE_END_ARGS()
 	
 	void Construct(const FArguments& InArgs);
+	
+	/** @return The search box that is currently being displayed. */
+	TSharedPtr<SSequencerSearchBox> GetActiveSearchBox() const;
 
 private:
 	
@@ -46,8 +50,10 @@ private:
 	/** Filter model for the Curve Editor. */
 	TSharedPtr<ILinkedFilterViewModel> FilterModel;
 	
-	/** The widget displaying the search bar, and other filter widgets. */
-	TSharedPtr<SSearchAndFilterWidget> FilterRow;
+	/** Content displayed when ELinkedFilterMode::Linked */
+	TSharedPtr<SSearchAndFilterWidget> LinkedFilterWidgets;
+	/** Content displayed when ELinkedFilterMode::Instanced */
+	TSharedPtr<SSearchAndFilterWidget> InstancedFilterWidgets;
 
 	/** @return The filtering widget. */
 	TSharedRef<SWidget> MakeFilterContent(const FArguments& InArgs);
diff --git a/Engine/Source/Editor/Sequencer/Private/MVVM/Views/CurveEditor/SCurveTreeContent.cpp b/Engine/Source/Editor/Sequencer/Private/MVVM/Views/CurveEditor/SCurveTreeContent.cpp
--- a/Engine/Source/Editor/Sequencer/Private/MVVM/Views/CurveEditor/SCurveTreeContent.cpp
+++ b/Engine/Source/Editor/Sequencer/Private/MVVM/Views/CurveEditor/SCurveTreeContent.cpp
@@ -2,17 +2,16 @@
 
 #include "SCurveTreeContent.h"
 
-#include "CurveEditor.h"
-#include "SCurveEditorSearchAndFilterRow.h"
-#include "SelectionSyncUtils.h"
 #include "Components/VerticalBox.h"
+#include "CurveEditor.h"
+#include "CurveEditor/Integration/CurveModelSyncer.h"
+#include "Filters/Widgets/SSequencerSearchBox.h"
 #include "Misc/TransportControlsHelper.h"
+#include "SCurveEditorSearchAndFilterRow.h"
 #include "SPopoutTabInlineContent.h"
 #include "STemporarilyFocusedSpinBox.h"
+#include "SelectionSyncUtils.h"
 #include "Sequencer.h"
-#include "CurveEditor/Integration/CurveModelSyncer.h"
-#include "Filters/Widgets/Linking/SFilterBarSwitcher.h"
-#include "MVVM/CurveEditorIntegrationExtension.h"
 #include "Styling/AppStyle.h"
 #include "Tree/SCurveEditorTree.h"
 #include "Widgets/Layout/SBorder.h"
@@ -54,8 +53,7 @@ void SCurveTreeContent::Construct(const FArguments& InArgs)
 
 void SCurveTreeContent::FocusSearchBox()
 {
-	// TODO DP
-	//FSlateApplication::Get().SetKeyboardFocus(CurveEditorSearchBox, EFocusCause::SetDirectly);
+	FSlateApplication::Get().SetKeyboardFocus(FilterRow->GetActiveSearchBox(), EFocusCause::SetDirectly);
 }
 
 void SCurveTreeContent::FocusPlayTimeDisplay()
@@ -85,7 +83,7 @@ TSharedRef<SWidget> SCurveTreeContent::MakeFilterArea(const FArguments& InArgs)
 			.BorderImage(FAppStyle::GetBrush("ToolPanel.GroupBorder"))
 			.Clipping(EWidgetClipping::ClipToBounds)
 			[
-				SNew(SCurveEditorSearchAndFilterRow)
+				SAssignNew(FilterRow, SCurveEditorSearchAndFilterRow)
 				.Sequencer(InArgs._Sequencer)
 				.FilterViewModel(InArgs._FilterViewModel)
 				.ScrollItemIntoView(this, &SCurveTreeContent::ScrollIntoView)
diff --git a/Engine/Source/Editor/Sequencer/Private/MVVM/Views/CurveEditor/SCurveTreeContent.h b/Engine/Source/Editor/Sequencer/Private/MVVM/Views/CurveEditor/SCurveTreeContent.h
--- a/Engine/Source/Editor/Sequencer/Private/MVVM/Views/CurveEditor/SCurveTreeContent.h
+++ b/Engine/Source/Editor/Sequencer/Private/MVVM/Views/CurveEditor/SCurveTreeContent.h
@@ -16,6 +16,7 @@ namespace UE::Sequencer
 {
 class FCurveModelSyncer;
 class FLinkedFilterViewModel;
+class SCurveEditorSearchAndFilterRow;
 
 /** Displays Curve Editor in Sequencer with linked filtering. */
 class SCurveTreeContent : public SCompoundWidget
@@ -62,6 +63,8 @@ private:
 	/** Used to translate between IOutlinerExtension and FCurveEditorTreeItemID. */
 	TAttribute<const FCurveModelSyncer*> CurveModelSyncerAttr;
 	
+	/** The filter row content displayed above the tree. */
+	TSharedPtr<SCurveEditorSearchAndFilterRow> FilterRow;
 	/** Curve editor tree widget */
 	TSharedPtr<SCurveEditorTree> CurveEditorTreeView;
 	/** The current playback time display. */
