[Sequencer] Turn off horizontal track separator lines off by default

#Synced-CL 51303498

diff --git a/Engine/Source/Editor/Sequencer/Private/SequencerSettings.cpp b/Engine/Source/Editor/Sequencer/Private/SequencerSettings.cpp
--- a/Engine/Source/Editor/Sequencer/Private/SequencerSettings.cpp
+++ b/Engine/Source/Editor/Sequencer/Private/SequencerSettings.cpp
@@ -64,7 +64,7 @@ USequencerSettings::USequencerSettings( const FObjectInitializer& ObjectInitiali
 	PlaybackRangeEndColor = FLinearColor(0.5f, 0.01f, 0.01f, 1.0f);
 	MarkedFrameColor = FLinearColor(0.f, 1.f, 1.f, 0.4f);
 	SectionColorAlpha = 1.0f;
-	bDrawHorizontalTrackSeparatorLines = true;
+	bDrawHorizontalTrackSeparatorLines = false;
 	bCleanPlaybackMode = true;
 	bActivateRealtimeViewports = true;
 	bEvaluateSubSequencesInIsolation = false;
