[Sequencer] Set scroll box widgets for outliner rows to fill

#Synced-CL 51303498

diff --git a/Engine/Source/Editor/Sequencer/Private/SSequencer.cpp b/Engine/Source/Editor/Sequencer/Private/SSequencer.cpp
--- a/Engine/Source/Editor/Sequencer/Private/SSequencer.cpp
+++ b/Engine/Source/Editor/Sequencer/Private/SSequencer.cpp
@@ -1098,6 +1098,8 @@ TSharedRef<SWidget> SSequencer::ConstructTransportControls(const TSharedRef<ISeq
 		.ScrollBarAlwaysVisible(false)
 		.ScrollBarVisibility(EVisibility::Collapsed)
 		+ SScrollBox::Slot()
+		.HAlign(HAlign_Fill)
+		.FillContentSize(1.f)
 		[
 			SNew(SHorizontalBox)
 			+ SHorizontalBox::Slot()
@@ -4212,6 +4214,8 @@ void SSequencer::RebuildSearchAndFilterRow()
 			.ScrollBarAlwaysVisible(false)
 			.ScrollBarVisibility(EVisibility::Collapsed)
 			+ SScrollBox::Slot()
+			.HAlign(HAlign_Fill)
+			.FillContentSize(1.f)
 			[
 				FilterRow.ToSharedRef()
 			]
