Sequencer: Handle rebinding Curve Editor commands to the active filter bar when the link mode changes.
All command bindings for curve editor are now centralized in FCurveCommandBinder.

This change addresses that FSequencerFilterBar affects the global FCurveEditor's command binding state each time it is instantiated. That's undesirable as multiple instances are now created.

#Synced-CL 51334859

diff --git a/Engine/Source/Editor/Sequencer/Private/CurveEditor/Integration/CurveEditorCommandBinder.cpp b/Engine/Source/Editor/Sequencer/Private/CurveEditor/Integration/CurveEditorCommandBinder.cpp
new file mode 100644
--- /dev/null
+++ b/Engine/Source/Editor/Sequencer/Private/CurveEditor/Integration/CurveEditorCommandBinder.cpp
@@ -0,0 +1,112 @@
+﻿// Copyright Epic Games, Inc. All Rights Reserved.
+
+#include "CurveEditorCommandBinder.h"
+
+#include "CurveEditor.h"
+#include "ISequencer.h"
+#include "Sequencer.h"
+#include "SequencerCommands.h"
+#include "Filters/SequencerTrackFilterCommands.h"
+#include "Filters/ViewModels/Linking/LinkedFilterViewModel.h"
+
+namespace UE::Sequencer
+{
+FCurveCommandBinder::FCurveCommandBinder(
+	const TSharedRef<FSequencer>& InSequencer,
+	const TSharedRef<FCurveEditor>& InCurveEditor,
+	const TSharedRef<FLinkedFilterViewModel>& InFilterViewModel
+	)
+	: WeakSequencer(InSequencer)
+	, CurveEditor(InCurveEditor)
+	, FilterViewModel(InFilterViewModel)
+{
+	FilterViewModel->OnFilterModeChanged().AddRaw(this, &FCurveCommandBinder::OnFilterModeChanged);
+}
+
+FCurveCommandBinder::~FCurveCommandBinder()
+{
+	FilterViewModel->OnFilterModeChanged().RemoveAll(this);
+}
+
+void FCurveCommandBinder::InitSequencerCommandBindings()
+{
+	const TSharedPtr<FSequencer> Sequencer = WeakSequencer.Pin();
+	if (!Sequencer)
+	{
+		return;
+	}
+	
+	const TSharedPtr<FUICommandList> CurveEditorSharedBindings = Sequencer->GetCommandBindings(ESequencerCommandBindings::CurveEditor);
+	const TSharedPtr<FUICommandList> SequencerCommandBindings = Sequencer->GetCommandBindings(ESequencerCommandBindings::Sequencer);
+	
+	const FSequencerCommands& Commands = FSequencerCommands::Get();
+	CurveEditorSharedBindings->MapAction(Commands.TogglePlay, *SequencerCommandBindings->GetActionForCommand(Commands.TogglePlay));
+	CurveEditorSharedBindings->MapAction(Commands.TogglePlayViewport, *SequencerCommandBindings->GetActionForCommand(Commands.TogglePlayViewport));
+	CurveEditorSharedBindings->MapAction(Commands.PlayForward, *SequencerCommandBindings->GetActionForCommand(Commands.PlayForward));
+	CurveEditorSharedBindings->MapAction(Commands.JumpToStart, *SequencerCommandBindings->GetActionForCommand(Commands.JumpToStart));
+	CurveEditorSharedBindings->MapAction(Commands.JumpToEnd, *SequencerCommandBindings->GetActionForCommand(Commands.JumpToEnd));
+	CurveEditorSharedBindings->MapAction(Commands.JumpToStartViewport, *SequencerCommandBindings->GetActionForCommand(Commands.JumpToStartViewport));
+	CurveEditorSharedBindings->MapAction(Commands.JumpToEndViewport, *SequencerCommandBindings->GetActionForCommand(Commands.JumpToEndViewport));
+	CurveEditorSharedBindings->MapAction(Commands.ShuttleBackward, *SequencerCommandBindings->GetActionForCommand(Commands.ShuttleBackward));
+	CurveEditorSharedBindings->MapAction(Commands.ShuttleForward, *SequencerCommandBindings->GetActionForCommand(Commands.ShuttleForward));
+	CurveEditorSharedBindings->MapAction(Commands.Pause, *SequencerCommandBindings->GetActionForCommand(Commands.Pause));
+	CurveEditorSharedBindings->MapAction(Commands.StepForward, *SequencerCommandBindings->GetActionForCommand(Commands.StepForward));
+	CurveEditorSharedBindings->MapAction(Commands.StepBackward, *SequencerCommandBindings->GetActionForCommand(Commands.StepBackward));
+	CurveEditorSharedBindings->MapAction(Commands.StepForwardViewport, *SequencerCommandBindings->GetActionForCommand(Commands.StepForwardViewport));
+	CurveEditorSharedBindings->MapAction(Commands.StepBackwardViewport, *SequencerCommandBindings->GetActionForCommand(Commands.StepBackwardViewport));
+	CurveEditorSharedBindings->MapAction(Commands.JumpForward, *SequencerCommandBindings->GetActionForCommand(Commands.JumpForward));
+	CurveEditorSharedBindings->MapAction(Commands.JumpBackward, *SequencerCommandBindings->GetActionForCommand(Commands.JumpBackward));
+	CurveEditorSharedBindings->MapAction(Commands.StepToNextKey, *SequencerCommandBindings->GetActionForCommand(Commands.StepToNextKey));
+	CurveEditorSharedBindings->MapAction(Commands.StepToPreviousKey, *SequencerCommandBindings->GetActionForCommand(Commands.StepToPreviousKey));
+	CurveEditorSharedBindings->MapAction(Commands.StepToNextMark, *SequencerCommandBindings->GetActionForCommand(Commands.StepToNextMark));
+	CurveEditorSharedBindings->MapAction(Commands.StepToPreviousMark, *SequencerCommandBindings->GetActionForCommand(Commands.StepToPreviousMark));
+	CurveEditorSharedBindings->MapAction(Commands.ToggleMarkAtPlayPosition, *SequencerCommandBindings->GetActionForCommand(Commands.ToggleMarkAtPlayPosition));
+	CurveEditorSharedBindings->MapAction(Commands.SetStartPlaybackRange, *SequencerCommandBindings->GetActionForCommand(Commands.SetStartPlaybackRange));
+	CurveEditorSharedBindings->MapAction(Commands.SetEndPlaybackRange, *SequencerCommandBindings->GetActionForCommand(Commands.SetEndPlaybackRange));
+
+	CurveEditorSharedBindings->MapAction(Commands.AddTransformKey, *SequencerCommandBindings->GetActionForCommand(Commands.AddTransformKey));
+	CurveEditorSharedBindings->MapAction(Commands.AddTranslationKey, *SequencerCommandBindings->GetActionForCommand(Commands.AddTranslationKey));
+	CurveEditorSharedBindings->MapAction(Commands.AddRotationKey, *SequencerCommandBindings->GetActionForCommand(Commands.AddRotationKey));
+	CurveEditorSharedBindings->MapAction(Commands.AddScaleKey, *SequencerCommandBindings->GetActionForCommand(Commands.AddScaleKey));
+	
+	CurveEditor->GetCommands()->Append(CurveEditorSharedBindings.ToSharedRef());
+	
+	BindCurveEditorToActiveFilterBar();
+}
+
+void FCurveCommandBinder::BindCurveEditorToActiveFilterBar() const
+{
+	const TSharedPtr<FSequencer> Sequencer = WeakSequencer.Pin();
+	const TSharedPtr<FUICommandList> CurveEditorBindings = Sequencer ? Sequencer->GetCommandBindings(ESequencerCommandBindings::CurveEditor) : nullptr;
+	if (!ensure(CurveEditorBindings))
+	{
+		return;
+	}
+	
+	const TSharedPtr<FSequencerFilterBar> FilterBar = FilterViewModel->GetActiveFilterBarImpl();
+	const TSharedPtr<FUICommandList> FilterBarCommandList = FilterBar->GetCommandList();
+	
+	// Add the general track filter commands
+	FSequencerTrackFilterCommands& TrackFilterCommands = FSequencerTrackFilterCommands::Get();
+	for (const TSharedPtr<FUICommandInfo>& Command : TrackFilterCommands.GetAllCommands())
+	{
+		if (Command.IsValid() && FilterBarCommandList->IsActionMapped(Command))
+		{
+			CurveEditorBindings->MapAction(Command, *FilterBarCommandList->GetActionForCommand(Command));
+		}
+	}
+
+	// Add the specific track filter toggle commands
+	const TArray<TSharedRef<FSequencerTrackFilter>> AllFilters = FilterBar->GetFilterList(true);
+	for (const TSharedRef<FSequencerTrackFilter>& Filter : AllFilters)
+	{
+		const TSharedPtr<FUICommandList>& FilterCommandList = Filter->GetFilterInterface().GetCommandList();
+		const TSharedPtr<FUICommandInfo>& FilterCommand = Filter->GetToggleCommand();
+
+		if (FilterCommand.IsValid() && FilterCommandList->IsActionMapped(FilterCommand))
+		{
+			CurveEditorBindings->MapAction(FilterCommand, *FilterCommandList->GetActionForCommand(FilterCommand));
+		}
+	}
+}
+}
diff --git a/Engine/Source/Editor/Sequencer/Private/CurveEditor/Integration/CurveEditorCommandBinder.h b/Engine/Source/Editor/Sequencer/Private/CurveEditor/Integration/CurveEditorCommandBinder.h
new file mode 100644
--- /dev/null
+++ b/Engine/Source/Editor/Sequencer/Private/CurveEditor/Integration/CurveEditorCommandBinder.h
@@ -0,0 +1,48 @@
+﻿// Copyright Epic Games, Inc. All Rights Reserved.
+
+#pragma once
+
+#include "Templates/SharedPointer.h"
+#include "Templates/UnrealTemplate.h"
+
+class FCurveEditor;
+class FSequencer;
+
+namespace UE::Sequencer
+{
+class FLinkedFilterViewModel;
+/** 
+ * Injects some Sequencer commands into Curve Editor such that when Curve Editor UI is focused, the user can continue to use some of the
+ * Sequencer commands, and they'll take affects as if they were issued while Sequencer UI was focused.
+ */
+class FCurveCommandBinder : public FNoncopyable
+{
+public:
+	
+	explicit FCurveCommandBinder(
+		const TSharedRef<FSequencer>& InSequencer, const TSharedRef<FCurveEditor>& InCurveEditor, 
+		const TSharedRef<FLinkedFilterViewModel>& InFilterViewModel
+		);
+	~FCurveCommandBinder();
+	
+	/** Binds Sequencer commands that should be exposed in Curve Editor. */
+	void InitSequencerCommandBindings();
+	
+private:
+	
+	/** The owning Sequencer whose commands are being rebound and that owns the FCurveEditor. */
+	const TWeakPtr<FSequencer> WeakSequencer;
+	
+	/** The curve editor to which we're binding the commands */
+	const TSharedRef<FCurveEditor> CurveEditor;
+	
+	/** When the ELinkedFilterMode changes, the bindings in the active filter bar must be reapplied to Curve Editor.*/
+	const TSharedRef<FLinkedFilterViewModel> FilterViewModel;
+	
+	
+	/** Handles any command rebinding that must take place when the active filter bar changes. */
+	void OnFilterModeChanged() const { BindCurveEditorToActiveFilterBar(); }
+	/** Remaps all curve editor bindings to invoke the command bindings of the active filter bar. */
+	void BindCurveEditorToActiveFilterBar() const;
+};
+} // namespace UE::Sequencer
diff --git a/Engine/Source/Editor/Sequencer/Private/CurveEditor/SequencerCurveEditorApp.cpp b/Engine/Source/Editor/Sequencer/Private/CurveEditor/SequencerCurveEditorApp.cpp
--- a/Engine/Source/Editor/Sequencer/Private/CurveEditor/SequencerCurveEditorApp.cpp
+++ b/Engine/Source/Editor/Sequencer/Private/CurveEditor/SequencerCurveEditorApp.cpp
@@ -90,6 +90,7 @@ static TSharedRef<FLinkedFilterViewModel> MakeFilteringModel(
 FSequencerCurveEditorApp::FSequencerCurveEditorApp(const TSharedRef<FSequencer>& InSequencer, const FTimeSliderArgs& InTimeSliderArgs)
 	: CurveEditor(CurveEditorDetail::MakeCurveEditor(InSequencer, InTimeSliderArgs))
 	, LinkedFilteringViewModel(CurveEditorDetail::MakeFilteringModel(InSequencer, InTimeSliderArgs, CurveEditor))
+	, CommandBinder(InSequencer, CurveEditor, LinkedFilteringViewModel)
 	, SelectionSyncer(InSequencer, CurveEditor, LinkedFilteringViewModel)
 	, WidgetOwner(
 		InSequencer, CurveEditor, 
@@ -107,7 +108,7 @@ FSequencerCurveEditorApp::FSequencerCurveEditorApp(const TSharedRef<FSequencer>&
 
 FSequencerCurveEditorApp* FSequencerCurveEditorApp::Get(const ISequencer& InSequencer)
 {
-	FCurveEditorExtension* CurveEditorExtension = InSequencer.GetViewModel()->GetRootModel()->CastDynamic<FCurveEditorExtension>();
+	const FCurveEditorExtension* CurveEditorExtension = InSequencer.GetViewModel()->CastDynamic<FCurveEditorExtension>();
 	return CurveEditorExtension ? CurveEditorExtension->CurveEditorApp.Get() : nullptr;
 }
 } // namespace UE::Sequencer
diff --git a/Engine/Source/Editor/Sequencer/Private/CurveEditor/SequencerCurveEditorApp.h b/Engine/Source/Editor/Sequencer/Private/CurveEditor/SequencerCurveEditorApp.h
--- a/Engine/Source/Editor/Sequencer/Private/CurveEditor/SequencerCurveEditorApp.h
+++ b/Engine/Source/Editor/Sequencer/Private/CurveEditor/SequencerCurveEditorApp.h
@@ -2,6 +2,7 @@
 
 #pragma once
 
+#include "Integration/CurveEditorCommandBinder.h"
 #include "Integration/CurveModelSyncer.h"
 #include "MVVM/Views/CurveEditor/CurveEditorWidgetOwner.h"
 #include "Templates/SharedPointer.h"
@@ -39,6 +40,9 @@ public:
 	 */
 	static FSequencerCurveEditorApp* Get(const ISequencer& InSequencer);
 	
+	/** Called when FSequencer::BindCommands() initializes its commands. */
+	void BindCommands() { CommandBinder.InitSequencerCommandBindings(); }
+	
 	/** Gets the curve editor view-model */
 	TSharedPtr<FCurveEditor> GetCurveEditor() const { return CurveEditor; }
 
@@ -61,12 +65,18 @@ public:
 	
 private:
 	
+	// @note Member order matters. C++ destruction order is reverse declaration order, so the dependencies must points upwards:
+	// Lower members can only depend on members declared after. Higher members cannot depend on members declared after.
+	
 	/** The Curve Editor instance that Sequencer uses. */
 	TSharedRef<FCurveEditor> CurveEditor;
 	
 	/** Representation of Curve Editor's linked filters. */
 	TSharedRef<FLinkedFilterViewModel> LinkedFilteringViewModel;
 	
+	/** Allows users to use issue some Sequencer commands when the Curve Editor is UI is focused. */
+	FCurveCommandBinder CommandBinder;
+	
 	/** Handles that Curve Editor has the same selection as that of Sequencer. */
 	FCurveModelSyncer SelectionSyncer;
 	
diff --git a/Engine/Source/Editor/Sequencer/Private/Filters/SequencerFilterBar.cpp b/Engine/Source/Editor/Sequencer/Private/Filters/SequencerFilterBar.cpp
--- a/Engine/Source/Editor/Sequencer/Private/Filters/SequencerFilterBar.cpp
+++ b/Engine/Source/Editor/Sequencer/Private/Filters/SequencerFilterBar.cpp
@@ -1,7 +1,7 @@
 // Copyright Epic Games, Inc. All Rights Reserved.
 
 #include "Filters/SequencerFilterBar.h"
-#include "CurveEditor.h"
+
 #include "Filters/Filters/SequencerTrackFilter_Keyed.h"
 #include "Filters/Filters/SequencerTrackFilter_Condition.h"
 #include "Filters/Filters/SequencerTrackFilter_CustomText.h"
@@ -24,7 +24,6 @@
 #include "Framework/Application/SlateApplication.h"
 #include "Framework/MultiBox/SToolBarButtonBlock.h"
 #include "MovieScene.h"
-#include "MVVM/CurveEditorExtension.h"
 #include "MVVM/ViewModels/CategoryModel.h"
 #include "MVVM/Views/SOutlinerView.h"
 #include "Sequencer.h"
@@ -205,13 +204,11 @@ void FSequencerFilterBar::CreateDefaultFilters()
 void FSequencerFilterBar::BindCommands()
 {
 	const FSequencerTrackFilterCommands& TrackFilterCommands = FSequencerTrackFilterCommands::Get();
-
-	const TSharedRef<SSequencer> SequencerWidget = StaticCastSharedRef<SSequencer>(GetSequencer().GetSequencerWidget());
-
+	
 	CommandList->MapAction(TrackFilterCommands.ToggleFilterBarVisibility,
-		FExecuteAction::CreateSP(SequencerWidget, &SSequencer::ToggleFilterBarVisibility),
+		FExecuteAction::CreateSP(this, &FSequencerFilterBar::ToggleFilterBarVisibility),
 		FCanExecuteAction(),
-		FIsActionChecked::CreateSP(SequencerWidget, &SSequencer::IsFilterBarVisible));
+		FIsActionChecked::CreateSP(this, &FSequencerFilterBar::IsFilterBarVisible));
 
 	CommandList->MapAction(TrackFilterCommands.SetToVerticalLayout,
 		FExecuteAction::CreateSP(this, &FSequencerFilterBar::SetToVerticalLayout),
@@ -261,44 +258,6 @@ void FSequencerFilterBar::BindCommands()
 	{
 		Filter->BindCommands();
 	}
-
-	// Add bindings for curve editor if supported
-	FCurveEditorExtension* const CurveEditorExtension = Sequencer.GetViewModel()->CastDynamic<FCurveEditorExtension>();
-	if (CurveEditorExtension)
-	{
-		const TSharedPtr<FCurveEditor> CurveEditor = CurveEditorExtension->GetCurveEditor();
-		if (ensure(CurveEditor.IsValid()))
-		{
-			const TSharedPtr<FUICommandList> CurveEditorCommands = CurveEditor->GetCommands();
-			if (ensure(CurveEditorCommands.IsValid()))
-			{
-				const TSharedPtr<FUICommandList> CurveEditorSharedBindings = Sequencer.GetCommandBindings(ESequencerCommandBindings::CurveEditor);
-
-				// Add the general track filter commands
-				for (const TSharedPtr<FUICommandInfo>& Command : TrackFilterCommands.GetAllCommands())
-				{
-					if (Command.IsValid() && CommandList->IsActionMapped(Command))
-					{
-						CurveEditorSharedBindings->MapAction(Command, *CommandList->GetActionForCommand(Command));
-					}
-				}
-
-				// Add the specific track filter toggle commands
-				for (const TSharedRef<FSequencerTrackFilter>& Filter : AllFilters)
-				{
-					const TSharedPtr<FUICommandList>& FilterCommandList = Filter->GetFilterInterface().GetCommandList();
-					const TSharedPtr<FUICommandInfo>& FilterCommand = Filter->GetToggleCommand();
-
-					if (FilterCommand.IsValid() && FilterCommandList->IsActionMapped(FilterCommand))
-					{
-						CurveEditorSharedBindings->MapAction(FilterCommand, *FilterCommandList->GetActionForCommand(FilterCommand));
-					}
-				}
-
-				CurveEditorCommands->Append(CurveEditorSharedBindings.ToSharedRef());
-			}
-		}
-	}
 }
 
 void FSequencerFilterBar::CreateCustomTextFiltersFromConfig()
diff --git a/Engine/Source/Editor/Sequencer/Private/Filters/ViewModels/Linking/LinkedFilterFactoryViewModel.cpp b/Engine/Source/Editor/Sequencer/Private/Filters/ViewModels/Linking/LinkedFilterFactoryViewModel.cpp
--- a/Engine/Source/Editor/Sequencer/Private/Filters/ViewModels/Linking/LinkedFilterFactoryViewModel.cpp
+++ b/Engine/Source/Editor/Sequencer/Private/Filters/ViewModels/Linking/LinkedFilterFactoryViewModel.cpp
@@ -20,6 +20,7 @@ FLinkedFilterFactoryViewModel::FLinkedFilterFactoryViewModel(
 TSharedRef<FLinkedFilterViewModel> FLinkedFilterFactoryViewModel::MakeFilteringModelImpl(const FMakeLinkedFilterViewModelArgs& InArgs) const
 {
 	const TSharedRef<FSequencerFilterBar> InstancedFilterBar = MakeShared<FSequencerFilterBar>(Sequencer, InArgs.ConfigSubIdentifier);
+	InstancedFilterBar->BindCommands();
 	return MakeShared<FLinkedFilterViewModel>(SharedFilterBar, InstancedFilterBar, FilterEvaluator);
 }
 
diff --git a/Engine/Source/Editor/Sequencer/Private/SSequencer.cpp b/Engine/Source/Editor/Sequencer/Private/SSequencer.cpp
--- a/Engine/Source/Editor/Sequencer/Private/SSequencer.cpp
+++ b/Engine/Source/Editor/Sequencer/Private/SSequencer.cpp
@@ -1139,7 +1139,7 @@ TSharedRef<SWidget> SSequencer::ConstructTransportControls(const TSharedRef<ISeq
 }
 UE_ENABLE_OPTIMIZATION_SHIP
 
-void SSequencer::BindCommands(TSharedRef<FUICommandList> SequencerCommandBindings, TSharedRef<FUICommandList> CurveEditorSharedBindings)
+void SSequencer::BindCommands(TSharedRef<FUICommandList> SequencerCommandBindings)
 {
 	auto CanPasteFromHistory = [this]{
 		if (!HasFocusedDescendants() && !HasKeyboardFocus())
diff --git a/Engine/Source/Editor/Sequencer/Private/SSequencer.h b/Engine/Source/Editor/Sequencer/Private/SSequencer.h
--- a/Engine/Source/Editor/Sequencer/Private/SSequencer.h
+++ b/Engine/Source/Editor/Sequencer/Private/SSequencer.h
@@ -296,7 +296,7 @@ public:
 
 	void Construct(const FArguments& InArgs, TSharedRef<FSequencer> InSequencer);
 
-	void BindCommands(TSharedRef<FUICommandList> SequencerCommandBindings, TSharedRef<FUICommandList> CurveEditorSharedBindings);
+	void BindCommands(TSharedRef<FUICommandList> SequencerCommandBindings);
 	
 	~SSequencer();
 	
diff --git a/Engine/Source/Editor/Sequencer/Private/Sequencer.cpp b/Engine/Source/Editor/Sequencer/Private/Sequencer.cpp
--- a/Engine/Source/Editor/Sequencer/Private/Sequencer.cpp
+++ b/Engine/Source/Editor/Sequencer/Private/Sequencer.cpp
@@ -12353,50 +12353,20 @@ void FSequencer::BindCommands()
 		Commands.ToggleTrackSelectionSolo,
 		FExecuteAction::CreateSP(this, &FSequencer::ToggleTrackSelectionSolo));
 
-	// If this sequencer supports a curve editor, let's add bindings for it.
-	FCurveEditorExtension* CurveEditorExtension = ViewModel->CastDynamic<FCurveEditorExtension>();
-	if (CurveEditorExtension && ensure(CurveEditorExtension->GetCurveEditor()))
+	// bind widget specific commands
+	SequencerWidget->BindCommands(SequencerCommandBindings);
+	
+	ViewModel->GetFilterMediator()->BindCommands(SequencerCommandBindings);
+	
+	// If this sequencer supports a curve editor, let's add bindings for it. 
+	// As FLinkedFilteringMediator::BindCommands binds the shared filter bar commands, which the curve editor also binds to, this must happen after.
+	if (FSequencerCurveEditorApp* CurveEditorIntegration = FSequencerCurveEditorApp::Get(*this))
 	{
 		// We want a subset of the commands to work in the Curve Editor too, but bound to our functions. This minimizes code duplication
 		// while also freeing us up from issues that result from Sequencer already using two lists (for which our commands might be spread
 		// across both lists which makes a direct copy like it already uses difficult).
-		CurveEditorSharedBindings->MapAction(Commands.TogglePlay, *SequencerCommandBindings->GetActionForCommand(Commands.TogglePlay));
-		CurveEditorSharedBindings->MapAction(Commands.TogglePlayViewport, *SequencerCommandBindings->GetActionForCommand(Commands.TogglePlayViewport));
-		CurveEditorSharedBindings->MapAction(Commands.PlayForward, *SequencerCommandBindings->GetActionForCommand(Commands.PlayForward));
-		CurveEditorSharedBindings->MapAction(Commands.JumpToStart, *SequencerCommandBindings->GetActionForCommand(Commands.JumpToStart));
-		CurveEditorSharedBindings->MapAction(Commands.JumpToEnd, *SequencerCommandBindings->GetActionForCommand(Commands.JumpToEnd));
-		CurveEditorSharedBindings->MapAction(Commands.JumpToStartViewport, *SequencerCommandBindings->GetActionForCommand(Commands.JumpToStartViewport));
-		CurveEditorSharedBindings->MapAction(Commands.JumpToEndViewport, *SequencerCommandBindings->GetActionForCommand(Commands.JumpToEndViewport));
-		CurveEditorSharedBindings->MapAction(Commands.ShuttleBackward, *SequencerCommandBindings->GetActionForCommand(Commands.ShuttleBackward));
-		CurveEditorSharedBindings->MapAction(Commands.ShuttleForward, *SequencerCommandBindings->GetActionForCommand(Commands.ShuttleForward));
-		CurveEditorSharedBindings->MapAction(Commands.Pause, *SequencerCommandBindings->GetActionForCommand(Commands.Pause));
-		CurveEditorSharedBindings->MapAction(Commands.StepForward, *SequencerCommandBindings->GetActionForCommand(Commands.StepForward));
-		CurveEditorSharedBindings->MapAction(Commands.StepBackward, *SequencerCommandBindings->GetActionForCommand(Commands.StepBackward));
-		CurveEditorSharedBindings->MapAction(Commands.StepForwardViewport, *SequencerCommandBindings->GetActionForCommand(Commands.StepForwardViewport));
-		CurveEditorSharedBindings->MapAction(Commands.StepBackwardViewport, *SequencerCommandBindings->GetActionForCommand(Commands.StepBackwardViewport));
-		CurveEditorSharedBindings->MapAction(Commands.JumpForward, *SequencerCommandBindings->GetActionForCommand(Commands.JumpForward));
-		CurveEditorSharedBindings->MapAction(Commands.JumpBackward, *SequencerCommandBindings->GetActionForCommand(Commands.JumpBackward));
-		CurveEditorSharedBindings->MapAction(Commands.StepToNextKey, *SequencerCommandBindings->GetActionForCommand(Commands.StepToNextKey));
-		CurveEditorSharedBindings->MapAction(Commands.StepToPreviousKey, *SequencerCommandBindings->GetActionForCommand(Commands.StepToPreviousKey));
-		CurveEditorSharedBindings->MapAction(Commands.StepToNextMark, *SequencerCommandBindings->GetActionForCommand(Commands.StepToNextMark));
-		CurveEditorSharedBindings->MapAction(Commands.StepToPreviousMark, *SequencerCommandBindings->GetActionForCommand(Commands.StepToPreviousMark));
-		CurveEditorSharedBindings->MapAction(Commands.ToggleMarkAtPlayPosition, *SequencerCommandBindings->GetActionForCommand(Commands.ToggleMarkAtPlayPosition));
-		CurveEditorSharedBindings->MapAction(Commands.SetStartPlaybackRange, *SequencerCommandBindings->GetActionForCommand(Commands.SetStartPlaybackRange));
-		CurveEditorSharedBindings->MapAction(Commands.SetEndPlaybackRange, *SequencerCommandBindings->GetActionForCommand(Commands.SetEndPlaybackRange));
-
-		CurveEditorSharedBindings->MapAction(Commands.AddTransformKey, *SequencerCommandBindings->GetActionForCommand(Commands.AddTransformKey));
-		CurveEditorSharedBindings->MapAction(Commands.AddTranslationKey, *SequencerCommandBindings->GetActionForCommand(Commands.AddTranslationKey));
-		CurveEditorSharedBindings->MapAction(Commands.AddRotationKey, *SequencerCommandBindings->GetActionForCommand(Commands.AddRotationKey));
-		CurveEditorSharedBindings->MapAction(Commands.AddScaleKey, *SequencerCommandBindings->GetActionForCommand(Commands.AddScaleKey));
-
-		TSharedPtr<FCurveEditor> CurveEditor = CurveEditorExtension->GetCurveEditor();
-		CurveEditor->GetCommands()->Append(CurveEditorSharedBindings);
+		CurveEditorIntegration->BindCommands();
 	}
-
-	// bind widget specific commands
-	SequencerWidget->BindCommands(SequencerCommandBindings, CurveEditorSharedBindings);
-	
-	ViewModel->GetFilterMediator()->BindCommands(SequencerCommandBindings);
 }
 
 void FSequencer::BuildAddTrackMenu(class FMenuBuilder& MenuBuilder)
