Sequencer: Sync Curve Editor to Sequencer selection.

Effective changes:
- bSynchronizeCurveEditorSelection: Select something in Sequencer Outliner-> select it in Curve Editor
- bSynchronizeOutlinerSelectionToCurveEditor: Select something in Curve Editor -> select in Sequencer Outliner
- bIsolateCurveEditorToSelection:
- Select something in Sequencer Outliner -> hide everything else in Curve Editor
- Select something in Curve Editor -> continue showing everything in Curve Editor.

#Synced-CL 51334859

diff --git a/Engine/Source/Editor/Sequencer/Private/CurveEditor/Integration/CurveModelSyncer.cpp b/Engine/Source/Editor/Sequencer/Private/CurveEditor/Integration/CurveModelSyncer.cpp
deleted file mode 100644
diff --git a/Engine/Source/Editor/Sequencer/Private/CurveEditor/Integration/CurveModelSyncer.h b/Engine/Source/Editor/Sequencer/Private/CurveEditor/Integration/CurveModelSyncer.h
deleted file mode 100644
diff --git a/Engine/Source/Editor/Sequencer/Private/CurveEditor/Integration/Selection/CurveModelSyncer.cpp b/Engine/Source/Editor/Sequencer/Private/CurveEditor/Integration/Selection/CurveModelSyncer.cpp
new file mode 100644
--- /dev/null
+++ b/Engine/Source/Editor/Sequencer/Private/CurveEditor/Integration/Selection/CurveModelSyncer.cpp
@@ -0,0 +1,299 @@
+﻿// Copyright Epic Games, Inc. All Rights Reserved.
+
+#include "CurveModelSyncer.h"
+
+#include "CurveEditor.h"
+#include "Sequencer.h"
+#include "Filters/ViewModels/Linking/ILinkedFilterViewModel.h"
+#include "Modification/Utils/ScopedSelectionChange.h"
+#include "MVVM/SharedViewModelData.h"
+#include "MVVM/ViewModelPtr.h"
+#include "MVVM/Extensions/IOutlinerExtension.h"
+#include "MVVM/Selection/Selection.h"
+#include "MVVM/ViewModels/ViewModelIterators.h"
+#include "MVVM/Views/CurveEditor/CurveEditorWidgetOwner.h"
+#include "Tree/CurveEditorTree.h"
+
+namespace UE::Sequencer
+{
+FCurveModelSyncer::FCurveModelSyncer(
+	const TSharedRef<FSequencer>& InSequencer,
+	const TSharedRef<FCurveEditor>& InCurveEditor, const TSharedRef<ILinkedFilterViewModel>& InLinkedFilterViewModel
+	)
+	: WeakSequencer(InSequencer)
+	, CurveEditor(InCurveEditor)
+	, LinkedFilterViewModel(InLinkedFilterViewModel)
+{
+	const TSharedPtr<FSequenceModel> OwnerModel = InSequencer->GetViewModel()->GetRootSequenceModel();
+	FSimpleMulticastDelegate& HierarchyChanged = OwnerModel->GetSharedData()->SubscribeToHierarchyChanged(OwnerModel);
+	HierarchyChanged.AddRaw(this, &FCurveModelSyncer::OnHierarchyChanged);
+	
+	InSequencer->GetNodeTree()->OnUpdated().AddRaw(this, &FCurveModelSyncer::OnNodeTreeUpdated);
+	InSequencer->GetViewModel()->GetSelection()->Outliner.OnChanged.AddRaw(this, &FCurveModelSyncer::OnOutlinerSelectionChanged);
+}
+
+FCurveModelSyncer::~FCurveModelSyncer()
+{
+	if (const TSharedPtr<FSequencer> Sequencer = WeakSequencer.Pin())
+	{
+		const TSharedPtr<FSequenceModel> OwnerModel = Sequencer->GetViewModel()->GetRootSequenceModel();
+		OwnerModel->GetSharedData()->UnsubscribeFromHierarchyChanged(OwnerModel, this);
+	
+		Sequencer->GetNodeTree()->OnUpdated().RemoveAll(this);
+		Sequencer->GetViewModel()->GetSelection()->Outliner.OnChanged.RemoveAll(this);
+	}
+	
+	// Prevent crash on shutdown. RequestSyncSelection sets a timer that keeps a FScopedTransaction. 
+	// ~FScopedTransaction expects GEditor to be valid, but it is null at shutdown.
+	CleanUpPendingSyncSelection();
+}
+
+static bool bSyncSelectionRequested = false;
+void FCurveModelSyncer::RequestSyncSelection()
+{
+	const TSharedPtr<FSequencer> Sequencer = WeakSequencer.Pin();
+	if (bSyncSelectionRequested || !Sequencer || SelectionSyncSuspension.IsSuspended())
+	{
+		return;
+	}
+	bSyncSelectionRequested = true;
+
+	// We schedule selection syncing to the next editor tick because we might want to select items that
+	// have just been added to the curve editor tree this tick. If it happened after the Slate update,
+	// these items don't yet have a UI widget, and so selecting them doesn't do anything.
+	//
+	// Note that we capture a weak pointer of our owner model because selection changes can happen
+	// right around the time when we want to unload everything (such as when loading a new map in the
+	// editor). We don't want to extend the lifetime of our stuff in that case.
+	TWeakPtr<FSequencerEditorViewModel> WeakRootViewModel = Sequencer->GetViewModel();
+
+	// Key selection supports undo. If RequestSyncSelection is called as part of an ongoing transaction, record the key selection change for undo.
+	const bool bShouldRecord = GUndo != nullptr;
+	// The current transaction needs to be extended until next tick...
+	TSharedPtr<FScopedTransaction> Transaction = bShouldRecord
+		? MakeShared<FScopedTransaction>(FText::GetEmpty(), bShouldRecord).ToSharedPtr() : nullptr;
+	// ... at which point we'll diff the changes.
+	TSharedPtr<CurveEditor::FScopedSelectionChange> SelectionChange = bShouldRecord
+		? MakeShared<CurveEditor::FScopedSelectionChange>(CurveEditor).ToSharedPtr() : nullptr;
+	
+	SyncSelectionHandle = GEditor->GetTimerManager()->SetTimerForNextTick(
+	[this, WeakRootViewModel, Transaction = MoveTemp(Transaction), SelectionChange = MoveTemp(SelectionChange)]() mutable
+		{
+			bSyncSelectionRequested = false;
+		
+			const TSharedPtr<FSequencerEditorViewModel> RootViewModel = WeakRootViewModel.Pin();
+			if (!RootViewModel.IsValid())
+			{
+				return;
+			}
+			const TSharedPtr<ISequencer> Sequencer = RootViewModel->GetSequencer();
+			if (!Sequencer)
+			{
+				return;
+			}
+				
+			SyncSelectionHandle.Invalidate();
+			// Notify UI to refresh the selection to that of Sequencer.
+			OnRequestSyncCurveEditorUIDelegate.Broadcast();
+
+			// Order matters: this appends a sub-transaction for the selection change...
+			SelectionChange.Reset();
+			// ... and only then can the parent transaction be closed. 
+			Transaction.Reset();
+		});
+}
+
+void FCurveModelSyncer::ResetCurveEditor()
+{
+	for (const TPair<TWeakPtr<FViewModel>, FCurveEditorTreeItemID>& Pair : ViewModelToTreeItemIDMap)
+	{
+		const FCurveEditorTreeItemID ItemID(Pair.Value);
+		const TViewModelPtr<ICurveEditorTreeItemExtension> ViewModel = CastViewModel<ICurveEditorTreeItemExtension>(Pair.Key.Pin());
+		if (ViewModel)
+		{
+			ViewModel->OnRemovedFromCurveEditor(CurveEditor);
+		}
+		CurveEditor->RemoveTreeItem(ItemID);
+	}
+
+	ViewModelToTreeItemIDMap.Reset();
+	TreeItemIDToViewModelMap.Reset();
+}
+
+TWeakPtr<FViewModel> FCurveModelSyncer::FindViewModelById(const FCurveEditorTreeItemID& InTreeId) const
+{
+	const TWeakPtr<FViewModel>* ViewModel = TreeItemIDToViewModelMap.Find(InTreeId);
+	return ViewModel ? *ViewModel : nullptr;
+}
+
+FCurveEditorTreeItemID FCurveModelSyncer::FindIdByViewModel(const TWeakPtr<FViewModel>& InViewModel) const
+{
+	const FCurveEditorTreeItemID* TreeItemID = ViewModelToTreeItemIDMap.Find(InViewModel);
+	return TreeItemID ? *TreeItemID : FCurveEditorTreeItemID::Invalid();
+}
+
+void FCurveModelSyncer::UpdateCurveEditor()
+{
+	const TSharedPtr<FSequencer> Sequencer = WeakSequencer.Pin();
+	if (!Sequencer || SelectionSyncSuspension.IsSuspended())
+	{
+		return;
+	}
+	
+	// TODO UE-364008: In the future, we can remove the paths with bUseFilteredOut.
+	const bool bUseFilteredOut = Sequencer->GetSequencerSettings()->GetLinkFiltersWithCurveEditor();
+	const auto IsFilteredOut = [this, bUseFilteredOut](const TViewModelPtr<IOutlinerExtension>& InItem)
+	{
+		return UseLinkedFilteringUI()
+			? LinkedFilterViewModel->IsFilteredOut(InItem)
+			: bUseFilteredOut && InItem && InItem->IsFilteredOut();
+	};
+
+	FCurveEditorTree* CurveEditorTree = CurveEditor->GetTree();
+
+	// Guard against multiple broadcasts here and defer them until the end of this function
+	FScopedCurveEditorTreeEventGuard ScopedEventGuard = CurveEditorTree->ScopedEventGuard();
+
+	// Remove any stale tree items
+	for (auto It = ViewModelToTreeItemIDMap.CreateIterator(); It; ++It)
+	{
+		bool bIsRelevant = false;
+		const FCurveEditorTreeItemID ItemID(It.Value());
+		TViewModelPtr<ICurveEditorTreeItemExtension> ViewModel = CastViewModel<ICurveEditorTreeItemExtension>(It.Key().Pin());
+		if (ViewModel)
+		{
+			TViewModelPtr<IOutlinerExtension> OutlinerModel = ViewModel.ImplicitCast();
+			const bool bIsVisible = !IsFilteredOut(OutlinerModel);
+			
+			const FCurveEditorTreeItem* TreeItem = CurveEditorTree->FindItem(ItemID);
+
+			FCurveEditorTreeItemID ParentID;
+			if (auto Parent = ViewModel.AsModel()->FindAncestorOfType<ICurveEditorTreeItemExtension>())
+			{
+				ParentID = ViewModelToTreeItemIDMap.FindRef(Parent.AsModel());
+			}
+			
+			bIsRelevant = bIsVisible && TreeItem && TreeItem->GetParentID() == ParentID;
+		}
+		
+		if (!bIsRelevant)
+		{
+			if (ViewModel)
+			{
+				ViewModel->OnRemovedFromCurveEditor(CurveEditor);
+			}
+			CurveEditor->RemoveTreeItem(ItemID);
+			It.RemoveCurrent();
+		}
+	}
+
+	// Do a second pass to remove any items that were removed recursively above
+	for (auto It = ViewModelToTreeItemIDMap.CreateIterator(); It; ++It)
+	{
+		if (CurveEditorTree->FindItem(It->Value) == nullptr)
+		{
+			TreeItemIDToViewModelMap.Remove(It->Value);
+			It.RemoveCurrent();
+		}
+	}
+
+	// Iterate all non-filtered out outliners and check for curve editor tree extensions
+	const bool bIncludeRootNode = false;
+	bool bItemsAdded = false;
+	const TSharedPtr<FSequenceModel> OwnerModel = Sequencer->GetViewModel()->GetRootSequenceModel();
+	for (TParentFirstChildIterator<IOutlinerExtension> It(OwnerModel, bIncludeRootNode); It; ++It)
+	{
+		if (IsFilteredOut(*It))
+		{
+			It.IgnoreCurrentChildren();
+			continue;
+		}
+
+		if (TViewModelPtr<ICurveEditorTreeItemExtension> ChildViewModel = (*It).ImplicitCast())
+		{
+			const FCurveEditorTreeItemID ItemID = ViewModelToTreeItemIDMap.FindRef(ChildViewModel.AsModel());
+			if (!ItemID.IsValid())
+			{
+				bItemsAdded = true;
+				AddToCurveEditor(ChildViewModel, CurveEditor);
+			}
+		}
+	}
+
+	// If new items have been added, synchronize selection after the view models have been added to the curve editor. Synchronization depends on curve model tree item IDs
+	if (bItemsAdded)
+	{
+		RequestSyncSelection();
+	}
+}
+
+FCurveEditorTreeItemID FCurveModelSyncer::AddToCurveEditor(
+	TViewModelPtr<ICurveEditorTreeItemExtension> InViewModel,
+	TSharedPtr<FCurveEditor> InCurveEditor
+	)
+{
+	// If the view model doesn't want to be in the curve editor, bail out
+	// Note that this means we will create curve editor items for each parent in the hierarchy up
+	// until the first parent that doesn't implement ICurveEditorTreeItemExtension
+	// That is: we don't create "dummy" entries when there's a "gap" in the hierarchy
+	if (!InViewModel)
+	{
+		return FCurveEditorTreeItemID::Invalid();
+	}
+
+	// Check if we already have a valid curve editor ID
+	if (FCurveEditorTreeItemID* Existing = ViewModelToTreeItemIDMap.Find(InViewModel.AsModel()))
+	{
+		if (InCurveEditor->GetTree()->FindItem(*Existing) != nullptr)
+		{
+			return *Existing;
+		}
+	}
+
+	// Recursively create any needed parent curve editor items
+	TViewModelPtr<ICurveEditorTreeItemExtension> Parent = InViewModel.AsModel()->CastParent<ICurveEditorTreeItemExtension>();
+
+	FCurveEditorTreeItemID ParentID = Parent ?
+		AddToCurveEditor(Parent, InCurveEditor) :
+		FCurveEditorTreeItemID::Invalid();
+
+	// Create the new curve editor item
+	FCurveEditorTreeItem* NewItem = InCurveEditor->AddTreeItem(ParentID);
+	TSharedPtr<ICurveEditorTreeItem> CurveEditorTreeItem = InViewModel->GetCurveEditorTreeItem();
+	NewItem->SetWeakItem(CurveEditorTreeItem);
+	TOptional<FString> UniquePathName =  InViewModel->GetUniquePathName();
+	NewItem->SetUniquePathName(UniquePathName);
+	
+	// Register the new ID in our map and notify the view model
+	ViewModelToTreeItemIDMap.Add(InViewModel.AsModel(), NewItem->GetID());
+	TreeItemIDToViewModelMap.Add(NewItem->GetID(), InViewModel.AsModel());
+	
+	InViewModel->OnAddedToCurveEditor(NewItem->GetID(), InCurveEditor);
+
+	return NewItem->GetID();
+}
+
+void FCurveModelSyncer::OnOutlinerSelectionChanged()
+{
+	const TSharedPtr<FSequencer> Sequencer = WeakSequencer.Pin();
+	USequencerSettings* Settings = Sequencer ? Sequencer->GetSequencerSettings() : nullptr;
+	if (Settings && Settings->ShouldSyncCurveEditorSelection())
+	{
+		// We schedule a selection synchronization for the next update. This synchronization must happen
+		// after all filters have been applied, because the items we want to select in the curve editor
+		// might be currently filtered out, but will be visible when filters are re-evaluated. This is
+		// why curve editor integration runs in FSequencerNodeTree after filtering.
+		RequestSyncSelection();
+	}
+}
+
+void FCurveModelSyncer::CleanUpPendingSyncSelection()
+{
+	bSyncSelectionRequested = false;
+	
+	if (GEditor)
+	{
+		GEditor->GetTimerManager()->ClearTimer(SyncSelectionHandle);
+	}
+}
+} // namespace UE::Sequencer
diff --git a/Engine/Source/Editor/Sequencer/Private/CurveEditor/Integration/Selection/CurveModelSyncer.h b/Engine/Source/Editor/Sequencer/Private/CurveEditor/Integration/Selection/CurveModelSyncer.h
new file mode 100644
--- /dev/null
+++ b/Engine/Source/Editor/Sequencer/Private/CurveEditor/Integration/Selection/CurveModelSyncer.h
@@ -0,0 +1,95 @@
+﻿// Copyright Epic Games, Inc. All Rights Reserved.
+
+#pragma once
+
+#include "Containers/Map.h"
+#include "Engine/TimerHandle.h"
+#include "Misc/Counter/ScopedSuspension.h"
+#include "MVVM/Extensions/ICurveEditorTreeItemExtension.h"
+#include "MVVM/ViewModels/ViewModel.h"
+#include "Templates/UnrealTemplate.h"
+
+class FSequencer;
+
+namespace UE::Sequencer
+{
+class ILinkedFilterViewModel;
+
+using FScopedSuspendSequencerToCurveEditorSync = FScopedSuspension;
+
+/** 
+ * Makes sure that Curve Editor has the same selection as that in Sequencer.
+ * Consequently, it knows which FCurveEditorTreeItemID and FViewModel correspond to each other.
+ * 
+ * Responsible for creating FCurveModel instances in FCurveEditor based on ICurveEditorTreeItemExtension in the Sequencer hierarchy.
+ * This relies on outliner items implementing ICurveEditorTreeItemExtension (or its default shim) if they want to show up in the curve editor.
+ */
+class FCurveModelSyncer : public FNoncopyable
+{
+public:
+	
+	explicit FCurveModelSyncer(
+		const TSharedRef<FSequencer>& InSequencer,
+		const TSharedRef<FCurveEditor>& InCurveEditor, 
+		const TSharedRef<ILinkedFilterViewModel>& InLinkedFilterViewModel
+		);
+	~FCurveModelSyncer();
+	
+	/** Synchronize curve editor selection with sequencer outliner selection on the next update. */
+	void RequestSyncSelection();
+	
+	/** Clears the curve editor of all contents. */
+	void ResetCurveEditor();
+	
+	/** @return The view model associated with the tree item. */
+	TWeakPtr<FViewModel> FindViewModelById(const FCurveEditorTreeItemID& InTreeId) const;
+	/** @return The tree item associated with the view model. */
+	FCurveEditorTreeItemID FindIdByViewModel(const TWeakPtr<FViewModel>&) const;
+	
+	/** Keeps the curve editor items up-to-date with the sequencer outliner by adding/removing  entries as needed. */
+	void UpdateCurveEditor();
+	
+	/** Suspends sync to Curve Editor for the lifetime of the scoped object. */
+	FScopedSuspendSequencerToCurveEditorSync SuspendSelectionSync() { return FScopedSuspendSequencerToCurveEditorSync(SelectionSyncSuspension); }
+	
+	/** Invoked to request that the UI be updated with the current selection. */
+	FSimpleMulticastDelegate& OnRequestSyncCurveEditorUI() { return OnRequestSyncCurveEditorUIDelegate; }
+	
+private:
+	
+	/** The owning Sequencer instance. Used to get settings. */
+	const TWeakPtr<FSequencer> WeakSequencer;
+	
+	/** The curve editor we are synchronizing to. */
+	const TSharedRef<FCurveEditor> CurveEditor;
+	/** Curve Editor selection must be refreshed when we switch filtering modes. */
+	const TSharedRef<ILinkedFilterViewModel> LinkedFilterViewModel;
+	
+	TMap<TWeakPtr<FViewModel>, FCurveEditorTreeItemID> ViewModelToTreeItemIDMap;
+	TMap<FCurveEditorTreeItemID, TWeakPtr<FViewModel>> TreeItemIDToViewModelMap;
+	
+	/** Handle to latent refresh of selection. */
+	FTimerHandle SyncSelectionHandle;
+	/** Invoked to request that the UI be updated with the current selection. */
+	FSimpleMulticastDelegate OnRequestSyncCurveEditorUIDelegate;
+	
+	/** Suspends syncing selection to Curve Editor. */
+	FSuspendCounter SelectionSyncSuspension;
+	
+	/** Adds the given view-model to the curve editor */
+	FCurveEditorTreeItemID AddToCurveEditor(TViewModelPtr<ICurveEditorTreeItemExtension> InViewModel, TSharedPtr<FCurveEditor> InCurveEditor);
+	
+	/** Update curve editor items when sequencer outliner items change */
+	void OnHierarchyChanged() { UpdateCurveEditor(); }
+	/** Updates the curve editor as result of the node tree hierarchy changing. */
+	void OnNodeTreeUpdated() { UpdateCurveEditor(); }
+	/** Syncs Curve Editor's selection if selection syncing is enabled. */
+	void OnOutlinerSelectionChanged();
+	
+	/** Cancels any pending transactions. */
+	void CleanUpPendingSyncSelection();
+};
+} // namespace UE::Sequencer
+
+
+
diff --git a/Engine/Source/Editor/Sequencer/Private/CurveEditor/Integration/Selection/SelectionSyncRootLogic.cpp b/Engine/Source/Editor/Sequencer/Private/CurveEditor/Integration/Selection/SelectionSyncRootLogic.cpp
new file mode 100644
--- /dev/null
+++ b/Engine/Source/Editor/Sequencer/Private/CurveEditor/Integration/Selection/SelectionSyncRootLogic.cpp
@@ -0,0 +1,89 @@
+﻿// Copyright Epic Games, Inc. All Rights Reserved.
+
+#include "SelectionSyncRootLogic.h"
+
+#include "CurveEditor.h"
+#include "Sequencer.h"
+#include "SequencerSelectionCurveFilter.h"
+#include "MVVM/Selection/Selection.h"
+
+namespace UE::Sequencer
+{
+FSelectionSyncRootLogic::FSelectionSyncRootLogic(
+	const TSharedRef<FSequencer>& InSequencer, const TSharedRef<FCurveEditor>& InCurveEditor, 
+	const TSharedRef<ILinkedFilterViewModel> InLinkedFilteringViewModel
+	)
+	: WeakSequencer(InSequencer)
+	, CurveEditor(InCurveEditor)
+	, SequencerToCurveEditorSyncer(InSequencer, InCurveEditor, InLinkedFilteringViewModel)
+	, CurveEditorToSequencerSyncer(SequencerToCurveEditorSyncer, InCurveEditor, InSequencer)
+{
+	SequencerToCurveEditorSyncer.OnRequestSyncCurveEditorUI().AddLambda([this]
+	{
+		const FScopedSuspendCurveEditorToSequencerSync Suspend = CurveEditorToSequencerSyncer.SuspendSync();
+		OnRequestRefreshCurveEditorUIDelegate.Broadcast();
+	});
+	
+	CurveEditorToSequencerSyncer.OnStartSync().AddRaw(this, &FSelectionSyncRootLogic::OnStartCurveEditorToSequencerSync);
+	CurveEditorToSequencerSyncer.OnEndSync().AddRaw(this, &FSelectionSyncRootLogic::OnEndCurveEditorToSequencerSync);
+	
+	InSequencer->GetViewModel()->GetSelection()->Outliner.OnChanged.AddRaw(this, &FSelectionSyncRootLogic::OnOutlinerSelectionChanged);
+	
+	SequencerToCurveEditorSyncer.UpdateCurveEditor();
+}
+
+FSelectionSyncRootLogic::~FSelectionSyncRootLogic()
+{
+	if (const TSharedPtr<FSequencer> Sequencer = WeakSequencer.Pin())
+	{
+		Sequencer->GetViewModel()->GetSelection()->Outliner.OnChanged.RemoveAll(this);
+	}
+}
+
+void FSelectionSyncRootLogic::UpdateIsolationFilter()
+{
+	const TSharedPtr<FSequencer> Sequencer = WeakSequencer.Pin(); 
+	const USequencerSettings* SequencerSettings = Sequencer ? Sequencer->GetSequencerSettings() : nullptr;
+	if (!SequencerSettings || bSkipIsolation)
+	{
+		return;
+	}
+	
+	// If we're isolating to the selection and there is one, add the filter
+	if (SequencerSettings->ShouldIsolateToCurveEditorSelection() && Sequencer->GetViewModel()->GetSelection()->Outliner.Num() != 0)
+	{
+		if (!SelectionCurveEditorFilter)
+		{
+			SelectionCurveEditorFilter = MakeShared<FSequencerSelectionCurveFilter>();
+		}
+		SelectionCurveEditorFilter->Update(Sequencer->GetViewModel()->GetSelection(), SequencerSettings->GetAutoExpandNodesOnSelection());
+		CurveEditor->GetTree()->AddFilter(SelectionCurveEditorFilter);
+	}
+	// If we're not isolating to the selection (or there is no selection) remove the filter
+	else if (SelectionCurveEditorFilter)
+	{
+		CurveEditor->GetTree()->RemoveFilter(SelectionCurveEditorFilter);
+		SelectionCurveEditorFilter.Reset();
+	}
+}
+
+void FSelectionSyncRootLogic::OnStartCurveEditorToSequencerSync()
+{
+	SuspendSequenceToCurveEditor.Emplace(SequencerToCurveEditorSyncer.SuspendSelectionSync());
+			
+	// We don't want Curve Editor items to disappear as result of the user selecting something in the Curve Editor tree.
+	bSkipIsolation = true;
+}
+
+void FSelectionSyncRootLogic::OnEndCurveEditorToSequencerSync()
+{
+	SuspendSequenceToCurveEditor.Reset();
+	bSkipIsolation = false;
+	
+	// Applying Curve Editor selection to the Outliner may cause items to become hidden in the Outliner, so refresh so the items are up to date. 
+	// Specifically,the "Selected" filter can hide items after changing the selection. 
+	// Example: Suppose that before you had a control rig selected: it would show all the controls, like chest, foot, torso, etc.
+	// Then, you select the chest control rig handle actor. The "Selected" filter hides foot, torso, etc. because they are no longer selected.
+	SequencerToCurveEditorSyncer.UpdateCurveEditor();
+}
+} // namespace UE::Sequencer
diff --git a/Engine/Source/Editor/Sequencer/Private/CurveEditor/Integration/Selection/SelectionSyncRootLogic.h b/Engine/Source/Editor/Sequencer/Private/CurveEditor/Integration/Selection/SelectionSyncRootLogic.h
new file mode 100644
--- /dev/null
+++ b/Engine/Source/Editor/Sequencer/Private/CurveEditor/Integration/Selection/SelectionSyncRootLogic.h
@@ -0,0 +1,76 @@
+﻿// Copyright Epic Games, Inc. All Rights Reserved.
+
+#pragma once
+
+#include "CurveModelSyncer.h"
+#include "Misc/Optional.h"
+#include "SequencerSelectionSyncer.h"
+#include "Templates/UnrealTemplate.h"
+
+class FCurveEditor;
+class FSequencer;
+
+namespace UE::Sequencer
+{
+struct FSequencerSelectionCurveFilter;
+class ILinkedFilterViewModel;
+
+/** Instantiates and mediates between all subfeatures to do with syncing selection between Sequencer's Outliner and Curve Editor. */
+class FSelectionSyncRootLogic : public FNoncopyable
+{
+public:
+	
+	explicit FSelectionSyncRootLogic(
+		const TSharedRef<FSequencer>& InSequencer,
+		const TSharedRef<FCurveEditor>& InCurveEditor, 
+		const TSharedRef<ILinkedFilterViewModel> InLinkedFilteringViewModel
+		);
+	~FSelectionSyncRootLogic();
+	
+	/** Synchronize curve editor selection with sequencer outliner selection on the next update. */
+	void RequestSyncSelection() { SequencerToCurveEditorSyncer.RequestSyncSelection(); }
+	/** Keeps the curve editor items up-to-date with the sequencer outliner by adding/removing  entries as needed. */
+	void UpdateCurveEditor() { SequencerToCurveEditorSyncer.UpdateCurveEditor(); }
+	/** Clears the curve editor of all contents. */
+	void ResetCurveEditor() { SequencerToCurveEditorSyncer.ResetCurveEditor(); }
+	
+	FCurveModelSyncer& GetSequencerToCurveEditorSyncer() { return SequencerToCurveEditorSyncer; }
+	
+	/** Invoked to request that the Curve Editor tree view UI updates its items. */
+	FSimpleMulticastDelegate& OnRequestRefreshCurveEditorUI() { return OnRequestRefreshCurveEditorUIDelegate; }
+	
+private:
+	
+	/** The owning Sequencer instance. */
+	const TWeakPtr<FSequencer> WeakSequencer;
+	/** The curve editor being synced. */
+	const TSharedRef<FCurveEditor> CurveEditor;
+	
+	/** Handles that Curve Editor has the same selection as that of Sequencer. */
+	FCurveModelSyncer SequencerToCurveEditorSyncer;
+	
+	/** Handles that Sequencer has the same selection as that of Curve Editor. */
+	FSequencerSelectionSyncer CurveEditorToSequencerSyncer;
+	/** Set while CurveEditorToSequencerSyncer is syncing. */
+	TOptional<FScopedSuspendSequencerToCurveEditorSync> SuspendSequenceToCurveEditor;
+	
+	/** Curve editor filter that shows only the nodes selected in the Outliner. Can be null. */
+	TSharedPtr<FSequencerSelectionCurveFilter> SelectionCurveEditorFilter;
+	/** 
+	 * Whether to skip isolating. 
+	 * 
+	 * Set when selection is being pushed from Curve Editor to Sequencer: 
+	 * We don't want Curve Editor items to disappear as result of the user selecting something in the Curve Editor tree.
+	 */
+	bool bSkipIsolation = false;
+	
+	/** Invoked to request that the Curve Editor tree view UI updates its items. */
+	FSimpleMulticastDelegate OnRequestRefreshCurveEditorUIDelegate;
+	
+	void OnOutlinerSelectionChanged() { UpdateIsolationFilter(); }
+	void UpdateIsolationFilter();
+	
+	void OnStartCurveEditorToSequencerSync();
+	void OnEndCurveEditorToSequencerSync();
+};
+} // namespace UE::Sequencer
diff --git a/Engine/Source/Editor/Sequencer/Private/CurveEditor/Integration/Selection/SequencerSelectionSyncer.cpp b/Engine/Source/Editor/Sequencer/Private/CurveEditor/Integration/Selection/SequencerSelectionSyncer.cpp
new file mode 100644
--- /dev/null
+++ b/Engine/Source/Editor/Sequencer/Private/CurveEditor/Integration/Selection/SequencerSelectionSyncer.cpp
@@ -0,0 +1,68 @@
+﻿// Copyright Epic Games, Inc. All Rights Reserved.
+
+#include "SequencerSelectionSyncer.h"
+
+#include "CurveEditor.h"
+#include "Sequencer.h"
+#include "MVVM/Selection/Selection.h"
+#include "MVVM/Views/CurveEditor/SelectionSyncUtils.h"
+#include "Misc/ScopeExit.h"
+#include "MVVM/Views/SOutlinerView.h"
+
+namespace UE::Sequencer
+{
+FSequencerSelectionSyncer::FSequencerSelectionSyncer(
+	const FCurveModelSyncer& InCurveModelSyncer,
+	const TSharedRef<FCurveEditor>& InCurveEditor,
+	TWeakPtr<FSequencer> InWeakSequencer
+	)
+	: CurveModelSyncer(InCurveModelSyncer)
+	, CurveEditor(InCurveEditor)
+	, WeakSequencer(MoveTemp(InWeakSequencer))
+{
+	if (FCurveEditorTree* Tree = CurveEditor->GetTree(); ensure(Tree))
+	{
+		Tree->Events.OnSelectionChanged.AddRaw(this, &FSequencerSelectionSyncer::OnCurveEditorSelectionChanged);
+	}
+}
+
+FSequencerSelectionSyncer::~FSequencerSelectionSyncer()
+{
+	if (FCurveEditorTree* Tree = CurveEditor->GetTree())
+	{
+		Tree->Events.OnSelectionChanged.RemoveAll(this);
+	}
+}
+
+void FSequencerSelectionSyncer::OnCurveEditorSelectionChanged()
+{
+	const TSharedPtr<FSequencer> Sequencer = WeakSequencer.Pin();
+	FCurveEditorTree* Tree = CurveEditor->GetTree();
+	
+	USequencerSettings* Settings = Sequencer ? Sequencer->GetSequencerSettings() : nullptr;
+	const bool bIsEnabled = Settings && Settings->ShouldSyncOutlinerSelectionToCurveEditor();
+	
+	if (SyncSuspension.IsSuspended() || !bIsEnabled || !ensure(Tree) || !ensure(Sequencer))
+	{
+		return;
+	}
+	
+	OnStartSyncDelegate.Broadcast();
+	ON_SCOPE_EXIT{ OnEndSyncDelegate.Broadcast(); };
+	
+	FSequencerSelection& TargetSelection = Sequencer->GetSelection();
+	ApplyCurveEditorToSequencer(TargetSelection, *Tree, CurveModelSyncer);
+	ExpandSelectedSequencerItems(TargetSelection);
+	
+	// This is needed to apply the expansion state by causing SOutlinerView::Refresh.
+	Sequencer->RefreshTree();
+	
+	// If the user clicked a single item, show it in the UI
+	const TSet<TWeakViewModelPtr<IOutlinerExtension>>& Selected = TargetSelection.Outliner.GetSelected();
+	if (Selected.Num() == 1)
+	{
+		const TWeakViewModelPtr<IOutlinerExtension> FirstSelected = *Selected.CreateConstIterator();
+		Sequencer->GetOutlinerViewWidget()->RequestScrollIntoView(FirstSelected);
+	}
+}
+} // namespace UE::Sequencer
diff --git a/Engine/Source/Editor/Sequencer/Private/CurveEditor/Integration/Selection/SequencerSelectionSyncer.h b/Engine/Source/Editor/Sequencer/Private/CurveEditor/Integration/Selection/SequencerSelectionSyncer.h
new file mode 100644
--- /dev/null
+++ b/Engine/Source/Editor/Sequencer/Private/CurveEditor/Integration/Selection/SequencerSelectionSyncer.h
@@ -0,0 +1,63 @@
+﻿// Copyright Epic Games, Inc. All Rights Reserved.
+
+#pragma once
+
+#include "Delegates/Delegate.h"
+#include "Misc/Counter/ScopedSuspension.h"
+#include "Misc/Counter/SuspendCounter.h"
+#include "Templates/SharedPointer.h"
+#include "Templates/UnrealTemplate.h"
+
+class FSequencer;
+class FCurveEditor;
+struct FCurveModelID;
+
+namespace UE::Sequencer
+{
+class FCurveModelSyncer;
+class FSequencerSelection;
+
+using FScopedSuspendCurveEditorToSequencerSync = FScopedSuspension;
+
+/** Syncs selection changes from Curve Editor to Sequencer. */
+class FSequencerSelectionSyncer : public FNoncopyable
+{
+public:
+	
+	explicit FSequencerSelectionSyncer(
+		const FCurveModelSyncer& InCurveModelSyncer UE_LIFETIMEBOUND,
+		const TSharedRef<FCurveEditor>& InCurveEditor, TWeakPtr<FSequencer> InWeakSequencer
+		);
+	~FSequencerSelectionSyncer();
+	
+	/** Suspends sync to Sequencer for the lifetime of the scoped object. */
+	[[nodiscard]] FScopedSuspendCurveEditorToSequencerSync SuspendSync() { return FScopedSuspendCurveEditorToSequencerSync(SyncSuspension); }
+	
+	/** Invoked before curve editor selection is synced to Sequencer's outliner. */
+	FSimpleMulticastDelegate& OnStartSync() { return OnStartSyncDelegate; }
+	/** Invoked after curve editor selection is synced to Sequencer's outliner. */
+	FSimpleMulticastDelegate& OnEndSync() { return OnEndSyncDelegate; }
+
+private:
+	
+	/** Used to translate curve model IDs to IOutlinerExtensions. */
+	const FCurveModelSyncer& CurveModelSyncer;
+	
+	/** Acts as source. Whenever selection changes in the tree, the selection is synced to Sequencer. */
+	const TSharedRef<FCurveEditor> CurveEditor;
+	
+	/** Used to get the selection and refresh the tree hierarchy when the expansion state changes. */
+	const TWeakPtr<FSequencer> WeakSequencer;
+	
+	/** Keeps track whether syncing is enabled. */
+	FSuspendCounter SyncSuspension;
+	
+	/** Invoked before curve editor selection is synced to Sequencer's outliner. */
+	FSimpleMulticastDelegate OnStartSyncDelegate;
+	/** Invoked after curve editor selection is synced to Sequencer's outliner. */
+	FSimpleMulticastDelegate OnEndSyncDelegate;
+	
+	/** Handles the selection changing in curve editor by applying it to Sequencer's outliner. */
+	void OnCurveEditorSelectionChanged();
+};
+} // namespace UE::Sequencer
diff --git a/Engine/Source/Editor/Sequencer/Private/CurveEditor/SequencerCurveEditorApp.cpp b/Engine/Source/Editor/Sequencer/Private/CurveEditor/SequencerCurveEditorApp.cpp
--- a/Engine/Source/Editor/Sequencer/Private/CurveEditor/SequencerCurveEditorApp.cpp
+++ b/Engine/Source/Editor/Sequencer/Private/CurveEditor/SequencerCurveEditorApp.cpp
@@ -91,19 +91,17 @@ FSequencerCurveEditorApp::FSequencerCurveEditorApp(const TSharedRef<FSequencer>&
 	: CurveEditor(CurveEditorDetail::MakeCurveEditor(InSequencer, InTimeSliderArgs))
 	, LinkedFilteringViewModel(CurveEditorDetail::MakeFilteringModel(InSequencer, InTimeSliderArgs, CurveEditor))
 	, CommandBinder(InSequencer, CurveEditor, LinkedFilteringViewModel)
-	, SelectionSyncer(InSequencer, CurveEditor, LinkedFilteringViewModel)
+	, SelectionSyncLogic(InSequencer, CurveEditor, LinkedFilteringViewModel)
 	, WidgetOwner(
 		InSequencer, CurveEditor, 
 		MakeShared<FSequencerCurveEditorTimeSliderController>(InTimeSliderArgs, InSequencer, CurveEditor), 
-		LinkedFilteringViewModel, SelectionSyncer
+		LinkedFilteringViewModel, SelectionSyncLogic.GetSequencerToCurveEditorSyncer()
 		)
 {
-	SelectionSyncer.OnRequestSyncCurveEditorUI().AddLambda([this]
+	SelectionSyncLogic.OnRequestRefreshCurveEditorUI().AddLambda([this]
 	{
 		WidgetOwner.SyncSelection();
 	});
-	
-	SelectionSyncer.UpdateCurveEditor();
 }
 
 FSequencerCurveEditorApp* FSequencerCurveEditorApp::Get(const ISequencer& InSequencer)
diff --git a/Engine/Source/Editor/Sequencer/Private/CurveEditor/SequencerCurveEditorApp.h b/Engine/Source/Editor/Sequencer/Private/CurveEditor/SequencerCurveEditorApp.h
--- a/Engine/Source/Editor/Sequencer/Private/CurveEditor/SequencerCurveEditorApp.h
+++ b/Engine/Source/Editor/Sequencer/Private/CurveEditor/SequencerCurveEditorApp.h
@@ -3,7 +3,7 @@
 #pragma once
 
 #include "Integration/CurveEditorCommandBinder.h"
-#include "Integration/CurveModelSyncer.h"
+#include "Integration/Selection/SelectionSyncRootLogic.h"
 #include "MVVM/Views/CurveEditor/CurveEditorWidgetOwner.h"
 #include "Templates/SharedPointer.h"
 #include "Templates/UnrealTemplate.h"
@@ -57,11 +57,11 @@ public:
 	TSharedPtr<SCurveEditorTree> GetCurveEditorTreeView() const { return WidgetOwner.GetTreeView(); }
 
 	/** Synchronize curve editor selection with sequencer outliner selection on the next update. */
-	void RequestSyncSelection() { SelectionSyncer.RequestSyncSelection(); }
+	void RequestSyncSelection() { SelectionSyncLogic.RequestSyncSelection(); }
 	/** Keeps the curve editor items up-to-date with the sequencer outliner by adding/removing  entries as needed. */
-	void UpdateCurveEditor() { SelectionSyncer.UpdateCurveEditor(); }
+	void UpdateCurveEditor() { SelectionSyncLogic.UpdateCurveEditor(); }
 	/** Clears the curve editor of all contents. */
-	void ResetCurveEditor() { SelectionSyncer.ResetCurveEditor(); }
+	void ResetCurveEditor() { SelectionSyncLogic.ResetCurveEditor(); }
 	
 private:
 	
@@ -77,8 +77,8 @@ private:
 	/** Allows users to use issue some Sequencer commands when the Curve Editor is UI is focused. */
 	FCurveCommandBinder CommandBinder;
 	
-	/** Handles that Curve Editor has the same selection as that of Sequencer. */
-	FCurveModelSyncer SelectionSyncer;
+	/** Instantiates and mediates between all subfeatures to do with syncing selection between Sequencer's Outliner and Curve Editor. */
+	FSelectionSyncRootLogic SelectionSyncLogic;
 	
 	/** Manages lifecycle of views exposed in Sequencer. */
 	FCurveEditorWidgetOwner WidgetOwner;
diff --git a/Engine/Source/Editor/Sequencer/Private/MVVM/Views/CurveEditor/SCurveTreeContent.cpp b/Engine/Source/Editor/Sequencer/Private/MVVM/Views/CurveEditor/SCurveTreeContent.cpp
--- a/Engine/Source/Editor/Sequencer/Private/MVVM/Views/CurveEditor/SCurveTreeContent.cpp
+++ b/Engine/Source/Editor/Sequencer/Private/MVVM/Views/CurveEditor/SCurveTreeContent.cpp
@@ -4,7 +4,6 @@
 
 #include "Components/VerticalBox.h"
 #include "CurveEditor.h"
-#include "CurveEditor/Integration/CurveModelSyncer.h"
 #include "Filters/Widgets/SSequencerSearchBox.h"
 #include "Misc/TransportControlsHelper.h"
 #include "SCurveEditorSearchAndFilterRow.h"
@@ -12,6 +11,7 @@
 #include "STemporarilyFocusedSpinBox.h"
 #include "SelectionSyncUtils.h"
 #include "Sequencer.h"
+#include "CurveEditor/Integration/Selection/CurveModelSyncer.h"
 #include "Styling/AppStyle.h"
 #include "Tree/SCurveEditorTree.h"
 #include "Widgets/Layout/SBorder.h"
@@ -66,9 +66,10 @@ void SCurveTreeContent::SyncSelection() const
 {
 	if (const TSharedPtr<FSequencer> Sequencer = WeakSequencer.Pin())
 	{
-		Sequencer::SyncSelection(
+		Sequencer::ApplySequencerToCurveEditorSelection(
 			Sequencer->GetViewModel()->GetSelection(), CurveEditor.ToSharedRef(), CurveEditorTreeView.ToSharedRef()
 			);
+		ExpandSelectedCurveEditorItems(*CurveEditor->GetTree(), *CurveEditorTreeView.ToSharedRef());
 	}
 }
 
diff --git a/Engine/Source/Editor/Sequencer/Private/MVVM/Views/CurveEditor/SOldCurveTreeContent.cpp b/Engine/Source/Editor/Sequencer/Private/MVVM/Views/CurveEditor/SOldCurveTreeContent.cpp
--- a/Engine/Source/Editor/Sequencer/Private/MVVM/Views/CurveEditor/SOldCurveTreeContent.cpp
+++ b/Engine/Source/Editor/Sequencer/Private/MVVM/Views/CurveEditor/SOldCurveTreeContent.cpp
@@ -2,6 +2,7 @@
 
 #include "SOldCurveTreeContent.h"
 
+#include "CurveEditor.h"
 #include "CurveEditorSettings.h"
 #include "LevelSequence.h"
 #include "SelectionSyncUtils.h"
@@ -76,9 +77,10 @@ void SOldCurveTreeContent::SyncSelection()
 {
 	if (const TSharedPtr<FSequencer> Sequencer = WeakSequencer.Pin())
 	{
-		Sequencer::SyncSelection(
+		ApplySequencerToCurveEditorSelection(
 			Sequencer->GetViewModel()->GetSelection(), CurveEditor.ToSharedRef(), CurveEditorTreeView.ToSharedRef()
 			);
+		ExpandSelectedCurveEditorItems(*CurveEditor->GetTree(), *CurveEditorTreeView);
 	}
 }
 
diff --git a/Engine/Source/Editor/Sequencer/Private/MVVM/Views/CurveEditor/SelectionSyncUtils.cpp b/Engine/Source/Editor/Sequencer/Private/MVVM/Views/CurveEditor/SelectionSyncUtils.cpp
--- a/Engine/Source/Editor/Sequencer/Private/MVVM/Views/CurveEditor/SelectionSyncUtils.cpp
+++ b/Engine/Source/Editor/Sequencer/Private/MVVM/Views/CurveEditor/SelectionSyncUtils.cpp
@@ -3,10 +3,12 @@
 #include "SelectionSyncUtils.h"
 
 #include "CurveEditor.h"
+#include "CurveEditor/Integration/Selection/CurveModelSyncer.h"
 #include "MVVM/Extensions/ICurveEditorTreeItemExtension.h"
 #include "MVVM/Extensions/IOutlinerExtension.h"
 #include "MVVM/Selection/Selection.h"
 #include "MVVM/ViewModelPtr.h"
+#include "Misc/ScopeExit.h"
 #include "Templates/SharedPointer.h"
 #include "Tree/SCurveEditorTree.h"
 
@@ -32,7 +34,7 @@ static void GetAllChildrenRecursively(
 }
 } // namespace SyncDetail
 
-void SyncSelection(
+void ApplySequencerToCurveEditorSelection(
 	const TSharedPtr<FSequencerSelection>& InSelection, 
 	const TSharedRef<FCurveEditor>& InCurveEditorModel, 
 	const TSharedRef<SCurveEditorTree>& InCurveEditorTreeView
@@ -41,6 +43,7 @@ void SyncSelection(
 	//we now sync to the first non-channel in the tree hierarchy, this is needed for hosted channels with control rigs
 	//we need to sort the hierarchy each time since the curve id's being used previously were not consistent in order of viewing/creation, persistent now for undo/redo
 	InCurveEditorModel->SuspendBroadcast();
+	ON_SCOPE_EXIT{ InCurveEditorModel->ResumeBroadcast(); };
 	InCurveEditorTreeView->ClearSelection();
 
 	//map of tree items id's to their actually sort order, need to use this since tree items ids are not in order anymore
@@ -121,7 +124,74 @@ void SyncSelection(
 	{
 		InCurveEditorTreeView->RequestScrollIntoView(FirstCurveEditorTreeItemID);
 	}
+}
 
-	InCurveEditorModel->ResumeBroadcast();
+namespace SyncDetail
+{
+static void SetParentsExpanded(const FCurveEditorTree& InTree, const FCurveEditorTreeItemID& InItem, SCurveEditorTree& InTreeView, bool bInExpanded)
+{
+	for (const FCurveEditorTreeItem* Item = InTree.FindItem(InItem); Item; Item = InTree.FindItem(Item->GetParentID()))
+	{
+		InTreeView.SetItemExpansion(Item->GetID(), bInExpanded);
+	}
+}
+}
+
+void ExpandSelectedCurveEditorItems(FCurveEditorTree& InTree, SCurveEditorTree& InTreeView)
+{
+	for (const TPair<FCurveEditorTreeItemID, ECurveEditorTreeSelectionState>& Pair : InTree.GetSelection())
+	{
+		if (Pair.Value != ECurveEditorTreeSelectionState::None)
+		{
+			SyncDetail::SetParentsExpanded(InTree, Pair.Key, InTreeView, true);
+		}
+	}
+}
+
+namespace SyncDetail
+{
+static void SetParentsExpanded(const TViewModelPtr<IOutlinerExtension>& InNode, const bool bInExpanded)
+{
+	for (const TViewModelPtr<IOutlinerExtension>& ParentNode : InNode.AsModel()->GetAncestorsOfType<IOutlinerExtension>())
+	{
+		if (ParentNode->IsExpanded() != bInExpanded)
+		{
+			ParentNode->SetExpansion(bInExpanded);
+		}
+	}
+}
+}
+
+void ApplyCurveEditorToSequencer(FSequencerSelection& InSelection, const FCurveEditorTree& InTree, const FCurveModelSyncer& InCurveModelSyncer)
+{
+	FOutlinerSelection NewOutlinerSelection;
+	const TMap<FCurveEditorTreeItemID, ECurveEditorTreeSelectionState>& CurveSelection = InTree.GetSelection();
+	for (const TPair<FCurveEditorTreeItemID, ECurveEditorTreeSelectionState>& Pair : CurveSelection)
+	{
+		if (Pair.Value != ECurveEditorTreeSelectionState::Explicit)
+		{
+			continue;
+		}
+		
+		const TSharedPtr<FViewModel> ViewModel = InCurveModelSyncer.FindViewModelById(Pair.Key).Pin();
+		const TViewModelPtr<IOutlinerExtension> OutlinerNode = ViewModel ? CastViewModel<IOutlinerExtension>(ViewModel.ToSharedRef()) : nullptr;
+		if (OutlinerNode)
+		{
+			NewOutlinerSelection.Select(OutlinerNode);
+		}
+	}
+	
+	InSelection.Outliner.ReplaceWith(MoveTemp(NewOutlinerSelection));
+}
+
+void ExpandSelectedSequencerItems(const FSequencerSelection& InSelection)
+{
+	for (const TWeakViewModelPtr<IOutlinerExtension>& WeakOutlinerNode : InSelection.Outliner.GetSelected())
+	{
+		if (const TViewModelPtr<IOutlinerExtension> OutlinerNode = WeakOutlinerNode.Pin())
+		{
+			SyncDetail::SetParentsExpanded(OutlinerNode, true);
+		}
+	}
 }
 } // namespace UE::Sequencer
diff --git a/Engine/Source/Editor/Sequencer/Private/MVVM/Views/CurveEditor/SelectionSyncUtils.h b/Engine/Source/Editor/Sequencer/Private/MVVM/Views/CurveEditor/SelectionSyncUtils.h
--- a/Engine/Source/Editor/Sequencer/Private/MVVM/Views/CurveEditor/SelectionSyncUtils.h
+++ b/Engine/Source/Editor/Sequencer/Private/MVVM/Views/CurveEditor/SelectionSyncUtils.h
@@ -4,17 +4,28 @@
 
 #include "Templates/SharedPointerFwd.h"
 
+class FCurveEditorTree;
 class FCurveEditor;
 class SCurveEditorTree;
 
 namespace UE::Sequencer
 {
+class FCurveModelSyncer;
 class FSequencerSelection;
 
-/** Updates SCurveEditorTree to have the same selection as InSelection */
-void SyncSelection(
+/** Applies selection in Sequencer to Curve Editor. */
+void ApplySequencerToCurveEditorSelection(
 	const TSharedPtr<FSequencerSelection>& InSelection, 
 	const TSharedRef<FCurveEditor>& InCurveEditorModel, 
 	const TSharedRef<SCurveEditorTree>& InCurveEditorTreeView
 	);
+
+/** Expands the tree hierarchy to show all items selected in Curve Editor. */
+void ExpandSelectedCurveEditorItems(FCurveEditorTree& InTree, SCurveEditorTree& InTreeView);
+
+/** Applies selection in Curve Editor to Sequencer.*/
+void ApplyCurveEditorToSequencer(FSequencerSelection& InSelection, const FCurveEditorTree& InTree, const FCurveModelSyncer& InCurveModelSyncer);
+
+/** Expands the tree hierarchy to show all items selected in Sequencer. */
+void ExpandSelectedSequencerItems(const FSequencerSelection& InSelection);
 } // namespace UE::Sequencer
diff --git a/Engine/Source/Editor/Sequencer/Private/Misc/Counter/ScopedSuspension.h b/Engine/Source/Editor/Sequencer/Private/Misc/Counter/ScopedSuspension.h
new file mode 100644
--- /dev/null
+++ b/Engine/Source/Editor/Sequencer/Private/Misc/Counter/ScopedSuspension.h
@@ -0,0 +1,40 @@
+﻿// Copyright Epic Games, Inc. All Rights Reserved.
+
+#pragma once
+
+#include "SuspendCounter.h"
+#include "Templates/UnrealTemplate.h"
+
+namespace UE::Sequencer
+{
+/** Increments FSuspendCounter for the lifetime of the instance.*/
+class FScopedSuspension : public FNoncopyable
+{
+public:
+	
+	explicit FScopedSuspension(FSuspendCounter& InCounter UE_LIFETIMEBOUND)
+		: Counter(&InCounter)
+	{
+		Counter->Increment();
+	}
+	
+	FScopedSuspension(FScopedSuspension&& InOther)
+		: Counter(InOther.Counter)
+	{
+		InOther.Counter = nullptr;
+	}
+	
+	~FScopedSuspension()
+	{
+		if (Counter)
+		{
+			Counter->Decrement();
+		}
+	}
+	
+private:
+	
+	FSuspendCounter* Counter;
+};
+}
+
diff --git a/Engine/Source/Editor/Sequencer/Private/Misc/Counter/SuspendCounter.h b/Engine/Source/Editor/Sequencer/Private/Misc/Counter/SuspendCounter.h
new file mode 100644
--- /dev/null
+++ b/Engine/Source/Editor/Sequencer/Private/Misc/Counter/SuspendCounter.h
@@ -0,0 +1,31 @@
+﻿// Copyright Epic Games, Inc. All Rights Reserved.
+
+#pragma once
+
+#include "HAL/Platform.h"
+#include "Math/UnrealMathUtility.h"
+
+namespace UE::Sequencer
+{
+/** 
+ * A counter that is used for suspending a certain action temporarily while the counter is positive.
+ * E.g. suspending event broadcast, etc.
+ */
+class FSuspendCounter
+{
+public:
+	
+	/** Starts suspending the action. */
+	void Increment() { ++Counter; }
+	/** Decrements the counter. If the counter reaches 0, the action is no longer suspended. */
+	void Decrement() { Counter = FMath::Max(0, Counter - 1); }
+	
+	/** @return Whether the action is suspended. */
+	bool IsSuspended() const { return Counter > 0; }
+	
+private:
+	
+	int32 Counter = 0;
+};
+}
+
diff --git a/Engine/Source/Editor/Sequencer/Private/SSequencer.cpp b/Engine/Source/Editor/Sequencer/Private/SSequencer.cpp
--- a/Engine/Source/Editor/Sequencer/Private/SSequencer.cpp
+++ b/Engine/Source/Editor/Sequencer/Private/SSequencer.cpp
@@ -1381,54 +1381,6 @@ void SSequencer::HandleKeySelectionChanged()
 
 void SSequencer::HandleOutlinerNodeSelectionChanged()
 {
-	using namespace UE::Sequencer;
-
-	TSharedPtr<FSequencer> Sequencer = SequencerPtr.Pin();
-	if (!Sequencer.IsValid())
-	{
-		return;
-	}
-
-	FViewModelPtr RootModel = Sequencer->GetViewModel();
-	FCurveEditorExtension* CurveEditorExtension = RootModel->CastDynamic<FCurveEditorExtension>();
-	if (!CurveEditorExtension)
-	{
-		return;
-	}
-
-	TSharedPtr<FCurveEditor> CurveEditor = CurveEditorExtension->GetCurveEditor();
-	if (CurveEditor)
-	{
-		const USequencerSettings* SequencerSettings = GetSequencerSettings();
-		// If we're isolating to the selection and there is one, add the filter
-		if (SequencerSettings->ShouldIsolateToCurveEditorSelection() && Sequencer->GetViewModel()->GetSelection()->Outliner.Num() != 0)
-		{
-			if (!SequencerSelectionCurveEditorFilter)
-			{
-				SequencerSelectionCurveEditorFilter = MakeShared<FSequencerSelectionCurveFilter>();
-			}
-
-			SequencerSelectionCurveEditorFilter->Update(Sequencer->GetViewModel()->GetSelection(), SequencerSettings->GetAutoExpandNodesOnSelection());
-
-			CurveEditor->GetTree()->AddFilter(SequencerSelectionCurveEditorFilter);
-		}
-		// If we're not isolating to the selection (or there is no selection) remove the filter
-		else if (SequencerSelectionCurveEditorFilter)
-		{
-			CurveEditor->GetTree()->RemoveFilter(SequencerSelectionCurveEditorFilter);
-			SequencerSelectionCurveEditorFilter = nullptr;
-		}
-
-		if (GetSequencerSettings()->ShouldSyncCurveEditorSelection())
-		{
-			// We schedule a selection synchronization for the next update. This synchronization must happen
-			// after all filters have been applied, because the items we want to select in the curve editor
-			// might be currently filtered out, but will be visible when filters are re-evaluated. This is
-			// why curve editor integration runs in FSequencerNodeTree after filtering.
-			CurveEditorExtension->RequestSyncSelection();
-		}
-	}
-
 	if (NodeGroupManager.IsValid())
 	{
 		NodeGroupManager->SelectItemsSelectedInSequencer();
diff --git a/Engine/Source/Editor/Sequencer/Private/SSequencer.h b/Engine/Source/Editor/Sequencer/Private/SSequencer.h
--- a/Engine/Source/Editor/Sequencer/Private/SSequencer.h
+++ b/Engine/Source/Editor/Sequencer/Private/SSequencer.h
@@ -756,9 +756,6 @@ private:
 	/** Section area widget for pinned tracks*/
 	TSharedPtr<UE::Sequencer::STrackAreaView> PinnedTrackArea;
 
-	/** Curve editor filter that shows only the selected nodes */
-	TSharedPtr<UE::Sequencer::FSequencerSelectionCurveFilter> SequencerSelectionCurveEditorFilter;
-
 	/** The breadcrumb trail widget for this sequencer */
 	TSharedPtr<SBreadcrumbTrail<FSequencerBreadcrumb>> BreadcrumbTrail;
 
diff --git a/Engine/Source/Editor/Sequencer/Private/Sequencer.cpp b/Engine/Source/Editor/Sequencer/Private/Sequencer.cpp
--- a/Engine/Source/Editor/Sequencer/Private/Sequencer.cpp
+++ b/Engine/Source/Editor/Sequencer/Private/Sequencer.cpp
@@ -2908,10 +2908,7 @@ void FSequencer::RecreateCurveEditor()
 	if (ensure(CurveEditorIntegration))
 	{
 		CurveEditorIntegration->ResetCurveEditor();
-		if (GetSequencerSettings()->ShouldSyncCurveEditorSelection())
-		{
-			CurveEditorIntegration->RequestSyncSelection();
-		}
+		CurveEditorIntegration->RequestSyncSelection();
 	}
 }
 
diff --git a/Engine/Source/Editor/Sequencer/Private/SequencerSettings.cpp b/Engine/Source/Editor/Sequencer/Private/SequencerSettings.cpp
--- a/Engine/Source/Editor/Sequencer/Private/SequencerSettings.cpp
+++ b/Engine/Source/Editor/Sequencer/Private/SequencerSettings.cpp
@@ -40,6 +40,7 @@ USequencerSettings::USequencerSettings( const FObjectInitializer& ObjectInitiali
 	bLinkCurveEditorTimeRange = false;
 	bLinkFiltersWithCurveEditor = false;
 	bSynchronizeCurveEditorSelection = true;
+	bSynchronizeOutlinerSelectionToCurveEditor = true;
 	bIsolateCurveEditorToSelection = true;
 	bCurveEditorVisible = false;
 	LoopMode = ESequencerLoopMode::SLM_NoLoop;
@@ -589,6 +590,15 @@ void USequencerSettings::SyncCurveEditorSelection(bool bInSynchronizeCurveEditor
 	}
 }
 
+void USequencerSettings::SyncOutlinerSelectionToCurveEditor(bool bInValue)
+{
+	if (bSynchronizeOutlinerSelectionToCurveEditor != bInValue)
+	{
+		bSynchronizeOutlinerSelectionToCurveEditor = bInValue;
+		SaveConfig();
+	}
+}
+
 void USequencerSettings::IsolateCurveEditorToSelection(bool bInIsolateCurveEditorToSelection)
 {
 	if (bIsolateCurveEditorToSelection != bInIsolateCurveEditorToSelection)
diff --git a/Engine/Source/Editor/Sequencer/Public/SequencerSettings.h b/Engine/Source/Editor/Sequencer/Public/SequencerSettings.h
--- a/Engine/Source/Editor/Sequencer/Public/SequencerSettings.h
+++ b/Engine/Source/Editor/Sequencer/Public/SequencerSettings.h
@@ -296,10 +296,15 @@ public:
 	/** Sets whether or not to link sequence filters with the curve editor tree */
 	UE_API void SetLinkFiltersWithCurveEditor(bool bInLinkFiltersWithCurveEditor);
 
-	/** Return true if we are to synchronize the curve editor and sequencer trees */
+	/** Return true whether changing the sequencer tree selection will also select the relevant nodes in the curve editor tree if possible */
 	bool ShouldSyncCurveEditorSelection() const { return bSynchronizeCurveEditorSelection; }
-	/** Assign whether we are to synchronize the curve editor and sequencer trees */
+	/** Assign whether changing the sequencer tree selection will also select the relevant nodes in the curve editor tree if possible */
 	UE_API void SyncCurveEditorSelection(bool bInSynchronizeCurveEditorSelection);
+	
+	/** Return true whether changing the curve editor tree selection will also select the relevant nodes in the outliner tree.*/
+	bool ShouldSyncOutlinerSelectionToCurveEditor() const { return bSynchronizeOutlinerSelectionToCurveEditor; }
+	/** Assign whether changing the curve editor tree selection will also select the relevant nodes in the outliner tree. */
+	UE_API void SyncOutlinerSelectionToCurveEditor(bool bInValue);
 
 	/** Return true if we should filter the curve editor tree to only nodes that are relevant to the current sequencer selection */
 	bool ShouldIsolateToCurveEditorSelection() const { return bIsolateCurveEditorToSelection; }
@@ -738,13 +743,17 @@ protected:
 	bool bLinkFiltersWithCurveEditor;
 
 	/** When enabled, changing the sequencer tree selection will also select the relevant nodes in the curve editor tree if possible. */
-	UPROPERTY( config, EditAnywhere, Category=CurveEditor )
+	UPROPERTY(config, EditAnywhere, Category=CurveEditor, meta = (DisplayName = "Synchronize Curve Editor Selection to Outliner"))
 	bool bSynchronizeCurveEditorSelection;
 
+	/** When enabled, changing the curve editor tree selection will also select the relevant nodes in the outliner tree. */
+	UPROPERTY(config, EditAnywhere, Category=CurveEditor)
+	bool bSynchronizeOutlinerSelectionToCurveEditor; 
+	
 	/** When enabled, changing the sequencer tree selection will isolate (auto-filter) the selected nodes in the curve editor. */
 	UPROPERTY(config, EditAnywhere, Category=CurveEditor)
 	bool bIsolateCurveEditorToSelection;
-
+	
 	/** Whether the curve editor is visible */
 	UPROPERTY(config, EditAnywhere, Category=CurveEditor)
 	bool bCurveEditorVisible;
