Sequencer: Remove vertical layout from Sequencer's filter bar.
This feature is not actively used so we decided to remove it to have simpler code - discounting future code iterations.

Notable changes:
- Deprecate ISequencerFilterBar (IsFilterBarLayout, SetToVerticalLayout, SetToHorizontalLayout, and ToggleFilterBarLayout) to not violate Liskov-Substitution principle.
- Remove EFilterBarLayout from SSequencerFilterBar, USequencerSettings, and FSequencerFilterBarConfig

For now, the navigation tool retains the capability for vertical layout. It can be removed separately if we decide it should only have horizontal layout.

#Synced-CL 51334859

diff --git a/Engine/Plugins/Experimental/SequenceNavigator/Source/SequenceNavigator/Private/Filters/NavigationToolFilterBar.cpp b/Engine/Plugins/Experimental/SequenceNavigator/Source/SequenceNavigator/Private/Filters/NavigationToolFilterBar.cpp
--- a/Engine/Plugins/Experimental/SequenceNavigator/Source/SequenceNavigator/Private/Filters/NavigationToolFilterBar.cpp
+++ b/Engine/Plugins/Experimental/SequenceNavigator/Source/SequenceNavigator/Private/Filters/NavigationToolFilterBar.cpp
@@ -698,18 +698,10 @@ void FNavigationToolFilterBar::ActivateAllEnabledFilters(const bool bInActivate,
 		}
 	}
 
-	if (bInActivate)
+	// Broadcast state change if any filter is active, or all filters are being deactivated.
+	if (!bInActivate || HasAnyFilterActive(false, false, true, true, true))
 	{
-		// Broadcast state change if any filter is active
-		if (HasAnyFilterActive(false, false, true, true, true))
-		{
-			BroadcastStateChanged();
-		}
-	}
-	else
-	{
-		// Broadcast state change if all filters are being deactivated
-		BroadcastStateChanged();
+		ActivationStateChangedEvent.Broadcast();
 	}
 }
 
@@ -1471,10 +1463,8 @@ void FNavigationToolFilterBar::ToggleFilterBarVisibility()
 	}
 
 	const bool bNewVisibility = !ToolSettings->IsFilterBarVisible();
-
 	ToolSettings->SetFilterBarVisible(bNewVisibility);
-
-	StateChangedEvent.Broadcast(bNewVisibility, ToolSettings->GetFilterBarLayout());
+	VisibilityChangedEvent.Broadcast(bNewVisibility);
 }
 
 bool FNavigationToolFilterBar::IsFilterBarLayout(const EFilterBarLayout InLayout) const
@@ -1495,8 +1485,7 @@ void FNavigationToolFilterBar::SetToVerticalLayout()
 	}
 
 	ToolSettings->SetFilterBarLayout(EFilterBarLayout::Vertical);
-
-	StateChangedEvent.Broadcast(IsFilterBarVisible(), ToolSettings->GetFilterBarLayout());
+	LayoutChangedEvent.Broadcast(ToolSettings->GetFilterBarLayout());
 }
 
 void FNavigationToolFilterBar::SetToHorizontalLayout()
@@ -1508,8 +1497,7 @@ void FNavigationToolFilterBar::SetToHorizontalLayout()
 	}
 
 	ToolSettings->SetFilterBarLayout(EFilterBarLayout::Horizontal);
-
-	StateChangedEvent.Broadcast(IsFilterBarVisible(), ToolSettings->GetFilterBarLayout());
+	LayoutChangedEvent.Broadcast(ToolSettings->GetFilterBarLayout());
 }
 
 void FNavigationToolFilterBar::ToggleFilterBarLayout()
@@ -1739,13 +1727,6 @@ void FNavigationToolFilterBar::OnSearchTextSaved(const FText& InFilterText)
 	CreateWindow_AddCustomTextFilter(CustomTextFilterData);
 }
 
-void FNavigationToolFilterBar::BroadcastStateChanged()
-{
-	const UNavigationToolSettings* const ToolSettings = GetDefault<UNavigationToolSettings>();
-	const EFilterBarLayout Layout = ToolSettings ? ToolSettings->GetFilterBarLayout() : EFilterBarLayout::Horizontal;
-	StateChangedEvent.Broadcast(IsFilterBarVisible(), Layout);
-}
-
 void FNavigationToolFilterBar::BroadcastFiltersChanged(const TSharedRef<FNavigationToolFilter>& InFilter, const ENavigationToolFilterChange InChangeType)
 {
 	FiltersChangedEvent.Broadcast(InChangeType, InFilter);
diff --git a/Engine/Plugins/Experimental/SequenceNavigator/Source/SequenceNavigator/Private/NavigationToolView.cpp b/Engine/Plugins/Experimental/SequenceNavigator/Source/SequenceNavigator/Private/NavigationToolView.cpp
--- a/Engine/Plugins/Experimental/SequenceNavigator/Source/SequenceNavigator/Private/NavigationToolView.cpp
+++ b/Engine/Plugins/Experimental/SequenceNavigator/Source/SequenceNavigator/Private/NavigationToolView.cpp
@@ -73,14 +73,17 @@ void FNavigationToolView::Init(const TSharedRef<FNavigationTool>& InTool, const
 	FilterBar = MakeShared<FNavigationToolFilterBar>(*InTool);
 	FilterBar->Init();
 	FilterBar->BindCommands(GetBaseCommandList());
-	FilterBar->OnStateChanged().AddSPLambda(this,
-		[this](const bool bInIsVisible, const EFilterBarLayout InNewLayout)
+	
+	const auto IssueRebuild = [this]
+	{
+		if (const TSharedPtr<SNavigationToolView> ToolWidget = StaticCastSharedPtr<SNavigationToolView>(GetToolWidget()))
 		{
-			if (const TSharedPtr<SNavigationToolView> ToolWidget = StaticCastSharedPtr<SNavigationToolView>(GetToolWidget()))
-			{
-				ToolWidget->RebuildWidget();
-			}
-		});
+			ToolWidget->RebuildWidget();
+		}
+	};
+	FilterBar->OnVisibilityChanged().AddSPLambda(this, [IssueRebuild](auto){ IssueRebuild(); });
+	FilterBar->OnLayoutChanged().AddSPLambda(this, [IssueRebuild](auto){ IssueRebuild(); });
+	FilterBar->OnActivationStateChanged().AddSPLambda(this, [IssueRebuild]{ IssueRebuild(); });
 
 	if (UNavigationToolSettings* const ToolSettings = WeakToolSettings.Get())
 	{
diff --git a/Engine/Plugins/Experimental/SequenceNavigator/Source/SequenceNavigator/Private/Widgets/SNavigationToolView.cpp b/Engine/Plugins/Experimental/SequenceNavigator/Source/SequenceNavigator/Private/Widgets/SNavigationToolView.cpp
--- a/Engine/Plugins/Experimental/SequenceNavigator/Source/SequenceNavigator/Private/Widgets/SNavigationToolView.cpp
+++ b/Engine/Plugins/Experimental/SequenceNavigator/Source/SequenceNavigator/Private/Widgets/SNavigationToolView.cpp
@@ -120,25 +120,15 @@ void SNavigationToolView::Construct(const FArguments& InArgs, const TSharedRef<F
 
 	if (const TSharedPtr<FNavigationToolFilterBar> FilterBar = InToolView->GetFilterBar())
 	{
-		FilterBar->OnStateChanged().AddSP(this, &SNavigationToolView::OnFilterBarStateChanged);
+		FilterBar->OnActivationStateChanged().AddSP(this, &SNavigationToolView::OnFilterActivationChanged);
+		FilterBar->OnVisibilityChanged().AddSP(this, &SNavigationToolView::OnFilterVisibilityChanged);
+		FilterBar->OnLayoutChanged().AddSP(this, &SNavigationToolView::OnFilterLayoutChanged);
 		FilterBar->OnFiltersChanged().AddSP(this, &SNavigationToolView::OnFiltersChanged);
 	}
 
 	RebuildWidget(true);
 }
 
-SNavigationToolView::~SNavigationToolView()
-{
-	if (const TSharedPtr<FNavigationToolView> ToolView = WeakToolView.Pin())
-	{
-		if (const TSharedPtr<FNavigationToolFilterBar> FilterBar = ToolView->GetFilterBar())
-		{
-			FilterBar->OnStateChanged().RemoveAll(this);
-			FilterBar->OnFiltersChanged().RemoveAll(this);
-		}
-	}
-}
-
 void SNavigationToolView::RebuildWidget(const bool bInReconstructColumns)
 {
 	bIsEnterLastKeyPressed = false;
@@ -644,13 +634,23 @@ void SNavigationToolView::SetFilterBarLayout(const EFilterBarLayout InLayout)
 	RebuildWidget();
 }
 
-void SNavigationToolView::OnFilterBarStateChanged(const bool InIsVisible, const EFilterBarLayout InNewLayout)
+void SNavigationToolView::OnFilterActivationChanged()
+{
+	RebuildWidget();
+}
+
+void SNavigationToolView::OnFilterVisibilityChanged(bool bInIsVisible)
+{
+	RebuildWidget();
+}
+
+void SNavigationToolView::OnFilterLayoutChanged(EFilterBarLayout InLayout)
 {
 	RebuildWidget();
 }
 
 void SNavigationToolView::OnFiltersChanged(const ENavigationToolFilterChange InChangeType
-	, const TSharedRef<FNavigationToolFilter>& InFilter)
+                                           , const TSharedRef<FNavigationToolFilter>& InFilter)
 {
 	RebuildWidget();
 }
diff --git a/Engine/Plugins/Experimental/SequenceNavigator/Source/SequenceNavigator/Private/Widgets/SNavigationToolView.h b/Engine/Plugins/Experimental/SequenceNavigator/Source/SequenceNavigator/Private/Widgets/SNavigationToolView.h
--- a/Engine/Plugins/Experimental/SequenceNavigator/Source/SequenceNavigator/Private/Widgets/SNavigationToolView.h
+++ b/Engine/Plugins/Experimental/SequenceNavigator/Source/SequenceNavigator/Private/Widgets/SNavigationToolView.h
@@ -32,8 +32,6 @@ public:
 
 	void Construct(const FArguments& InArgs, const TSharedRef<FNavigationToolView>& InToolView);
 
-	virtual ~SNavigationToolView() override;
-
 	TSharedPtr<SNavigationToolTreeView> GetTreeView() const { return TreeView; }
 
 	/** Rebuilds the entire widget. Needs to be called when the filter bar layout changes. */
@@ -88,9 +86,10 @@ protected:
 
 	void RebuildSearchAndFilterRow();
 
-	void OnFilterBarStateChanged(const bool InIsVisible, const EFilterBarLayout InNewLayout);
-	void OnFiltersChanged(const ENavigationToolFilterChange InChangeType
-		, const TSharedRef<FNavigationToolFilter>& InFilter);
+	void OnFilterActivationChanged();
+	void OnFilterVisibilityChanged(bool bInIsVisible);
+	void OnFilterLayoutChanged(EFilterBarLayout InLayout);
+	void OnFiltersChanged(const ENavigationToolFilterChange InChangeType, const TSharedRef<FNavigationToolFilter>& InFilter);
 
 	TSharedPtr<FNavigationToolFilterBar> GetFilterBar() const;
 
diff --git a/Engine/Plugins/Experimental/SequenceNavigator/Source/SequenceNavigator/Public/Filters/NavigationToolFilterBar.h b/Engine/Plugins/Experimental/SequenceNavigator/Source/SequenceNavigator/Public/Filters/NavigationToolFilterBar.h
--- a/Engine/Plugins/Experimental/SequenceNavigator/Source/SequenceNavigator/Public/Filters/NavigationToolFilterBar.h
+++ b/Engine/Plugins/Experimental/SequenceNavigator/Source/SequenceNavigator/Public/Filters/NavigationToolFilterBar.h
@@ -43,8 +43,8 @@ public:
 	static const FName SharedIdentifier;
 
 	DECLARE_EVENT_TwoParams(ISequencerFilterBar, FNavigationToolFiltersChanged, const ENavigationToolFilterChange /*InChangeType*/, const TSharedRef<FNavigationToolFilter>& /*InFilter*/);
-	DECLARE_EVENT_TwoParams(ISequencerFilterBar, FNavigationToolCustomTextFiltersChanged, const ENavigationToolFilterChange /*InChangeType*/, const TSharedRef<FNavigationToolFilter_CustomText>& /*InFilter*/);
-
+	DECLARE_EVENT_OneParam(ISequencerFilterBar, FOnLayoutChanged, EFilterBarLayout /** InNewLayout */);
+	
 	static FCustomTextFilterData DefaultNewCustomTextFilterData(const FText& InFilterString);
 
 	FNavigationToolFilterBar(FNavigationTool& InTool);
@@ -68,7 +68,9 @@ public:
 	virtual void ResetFilters() override;
 
 	FNavigationToolFiltersChanged& OnFiltersChanged() { return FiltersChangedEvent; }
-	virtual FOnFilterBarStateChanged& OnStateChanged() override { return StateChangedEvent; }
+	FOnLayoutChanged& OnLayoutChanged() { return LayoutChangedEvent; }
+	FSimpleMulticastDelegate& OnActivationStateChanged() { return ActivationStateChangedEvent; }
+	virtual FOnVisibilityChanged& OnVisibilityChanged() override { return VisibilityChangedEvent; }
 	virtual FSimpleMulticastDelegate& OnRequestUpdate() override { return RequestUpdateEvent; }
 	virtual FSimpleMulticastDelegate& OnTextFilterTextChanged() override { return TextFilterTextChangedEvent; }
 	virtual FMuteStateChangedDelegate& OnMuteFiltersChanged() override { return MuteFiltersChangedEvent; }
@@ -194,10 +196,10 @@ public:
 	virtual bool IsFilterBarVisible() const override;
 	virtual void ToggleFilterBarVisibility() override;
 
-	virtual bool IsFilterBarLayout(const EFilterBarLayout InLayout) const override;
-	virtual void SetToVerticalLayout() override;
-	virtual void SetToHorizontalLayout() override;
-	virtual void ToggleFilterBarLayout() override;
+	bool IsFilterBarLayout(const EFilterBarLayout InLayout) const;
+	void SetToVerticalLayout();
+	void SetToHorizontalLayout();
+	void ToggleFilterBarLayout();
 
 	virtual void CreateWindow_AddCustomTextFilter(const FCustomTextFilterData& InCustomTextFilterData = FCustomTextFilterData()) override;
 	virtual void CreateWindow_EditCustomTextFilter(const TSharedPtr<FNavigationToolFilter_CustomText>& InCustomTextFilter) override;
@@ -231,7 +233,6 @@ protected:
 	void OnSearchTextCommitted(const FText& InFilterText, const ETextCommit::Type InCommitType);
 	void OnSearchTextSaved(const FText& InFilterText);
 
-	void BroadcastStateChanged();
 	void BroadcastFiltersChanged(const TSharedRef<FNavigationToolFilter>& InFilter, const ENavigationToolFilterChange InChangeType);
 
 	/** The Navigation Tool this filter bar is interacting with */
@@ -263,7 +264,9 @@ protected:
 	TSharedPtr<SNavigationToolFilterBar> FilterBarWidget;
 
 	FNavigationToolFiltersChanged FiltersChangedEvent;
-	FOnFilterBarStateChanged StateChangedEvent;
+	FOnLayoutChanged LayoutChangedEvent;
+	FSimpleMulticastDelegate ActivationStateChangedEvent;
+	FOnVisibilityChanged VisibilityChangedEvent;
 	FSimpleMulticastDelegate RequestUpdateEvent;
 	FSimpleMulticastDelegate TextFilterTextChangedEvent;
 	FMuteStateChangedDelegate MuteFiltersChangedEvent;
diff --git a/Engine/Source/Editor/Sequencer/Private/Filters/Menus/SequencerViewOptionsMenu.cpp b/Engine/Source/Editor/Sequencer/Private/Filters/Menus/SequencerViewOptionsMenu.cpp
--- a/Engine/Source/Editor/Sequencer/Private/Filters/Menus/SequencerViewOptionsMenu.cpp
+++ b/Engine/Source/Editor/Sequencer/Private/Filters/Menus/SequencerViewOptionsMenu.cpp
@@ -144,55 +144,8 @@ void FSequencerViewOptionsMenu::PopulateFilterOptionsSection(UToolMenu& InMenu)
 		),
 		EUserInterfaceActionType::ToggleButton);
 
-	OptionsSection.AddSubMenu(TEXT("FilterBarOptions")
-		, LOCTEXT("FilterBarOptionsSubMenu", "Filter Bar Options")
-		, LOCTEXT("FilterBarOptionsSubMenuTooltip", "Filter Bar Options")
-		, FNewToolMenuDelegate::CreateSP(this, &FSequencerViewOptionsMenu::CreateFilterBarOptionsSection));
-}
-
-void FSequencerViewOptionsMenu::CreateFilterBarOptionsSection(UToolMenu* const InMenu)
-{
-	if (!InMenu)
-	{
-		return;
-	}
-
 	const FSequencerTrackFilterCommands& TrackFilterCommands = FSequencerTrackFilterCommands::Get();
-
-	FToolMenuSection& FilterBarVisibilitySection = InMenu->FindOrAddSection(TEXT("FilterBarVisibility")
-		, LOCTEXT("FilterBarVisibilityHeading", "Filter Bar"));
-
-	FilterBarVisibilitySection.AddMenuEntry(TrackFilterCommands.ToggleFilterBarVisibility);
-
-	FToolMenuSection& FilterBarLayoutSection = InMenu->FindOrAddSection(TEXT("FilterBarLayout")
-		, LOCTEXT("FilterBarLayoutHeading", "Filter Bar Layout"));
-
-	FilterBarLayoutSection.AddMenuEntry(TrackFilterCommands.SetToVerticalLayout);
-	FilterBarLayoutSection.AddMenuEntry(TrackFilterCommands.SetToHorizontalLayout);
-}
-
-bool FSequencerViewOptionsMenu::IsFilterLayout(const EFilterBarLayout InLayout) const
-{
-	const TSharedPtr<FSequencer> Sequencer = WeakSequencer.Pin();
-	if (!Sequencer.IsValid())
-	{
-		return false;
-	}
-
-	const TSharedRef<SSequencer> SequencerWidget = StaticCastSharedRef<SSequencer>(Sequencer->GetSequencerWidget());
-	return InLayout == SequencerWidget->GetFilterBarLayout();
-}
-
-void FSequencerViewOptionsMenu::SetFilterLayout(const EFilterBarLayout InLayout)
-{
-	const TSharedPtr<FSequencer> Sequencer = WeakSequencer.Pin();
-	if (!Sequencer.IsValid())
-	{
-		return;
-	}
-
-	const TSharedRef<SSequencer> SequencerWidget = StaticCastSharedRef<SSequencer>(Sequencer->GetSequencerWidget());
-	SequencerWidget->SetFilterBarLayout(InLayout);
+	OptionsSection.AddMenuEntry(TrackFilterCommands.ToggleFilterBarVisibility);
 }
 
 bool FSequencerViewOptionsMenu::IsIncludePinnedInFilter() const
diff --git a/Engine/Source/Editor/Sequencer/Private/Filters/Menus/SequencerViewOptionsMenu.h b/Engine/Source/Editor/Sequencer/Private/Filters/Menus/SequencerViewOptionsMenu.h
--- a/Engine/Source/Editor/Sequencer/Private/Filters/Menus/SequencerViewOptionsMenu.h
+++ b/Engine/Source/Editor/Sequencer/Private/Filters/Menus/SequencerViewOptionsMenu.h
@@ -21,11 +21,6 @@ protected:
 	void PopulateSortAndOrganizeSection(UToolMenu& InMenu);
 	void PopulateFilterOptionsSection(UToolMenu& InMenu);
 
-	void CreateFilterBarOptionsSection(UToolMenu* const InMenu);
-
-	bool IsFilterLayout(const EFilterBarLayout InLayout) const;
-	void SetFilterLayout(const EFilterBarLayout InLayout);
-
 	bool IsIncludePinnedInFilter() const;
 	void ToggleIncludePinnedInFilter();
 
diff --git a/Engine/Source/Editor/Sequencer/Private/Filters/SequencerFilterBar.cpp b/Engine/Source/Editor/Sequencer/Private/Filters/SequencerFilterBar.cpp
--- a/Engine/Source/Editor/Sequencer/Private/Filters/SequencerFilterBar.cpp
+++ b/Engine/Source/Editor/Sequencer/Private/Filters/SequencerFilterBar.cpp
@@ -210,16 +210,6 @@ void FSequencerFilterBar::BindCommands()
 		FCanExecuteAction(),
 		FIsActionChecked::CreateSP(this, &FSequencerFilterBar::IsFilterBarVisible));
 
-	CommandList->MapAction(TrackFilterCommands.SetToVerticalLayout,
-		FExecuteAction::CreateSP(this, &FSequencerFilterBar::SetToVerticalLayout),
-		FCanExecuteAction(),
-		FIsActionChecked::CreateSP(this, &FSequencerFilterBar::IsFilterBarLayout, EFilterBarLayout::Vertical));
-
-	CommandList->MapAction(TrackFilterCommands.SetToHorizontalLayout,
-		FExecuteAction::CreateSP(this, &FSequencerFilterBar::SetToHorizontalLayout),
-		FCanExecuteAction(),
-		FIsActionChecked::CreateSP(this, &FSequencerFilterBar::IsFilterBarLayout, EFilterBarLayout::Horizontal));
-
 	CommandList->MapAction(TrackFilterCommands.ResetFilters,
 		FExecuteAction::CreateSP(this, &FSequencerFilterBar::ResetFilters),
 		FCanExecuteAction::CreateSP(this, &FSequencerFilterBar::CanResetFilters));
@@ -1659,57 +1649,8 @@ void FSequencerFilterBar::ToggleFilterBarVisibility()
 	}
 
 	const bool bNewVisibility = !Settings->IsFilterBarVisible();
-
 	Settings->SetFilterBarVisible(bNewVisibility);
-
-	StateChangedEvent.Broadcast(bNewVisibility, Settings->GetFilterBarLayout());
-}
-
-bool FSequencerFilterBar::IsFilterBarLayout(const EFilterBarLayout InLayout) const
-{
-	if (const USequencerSettings* const Settings = Sequencer.GetSequencerSettings())
-	{
-		return (Settings->GetFilterBarLayout() == InLayout);
-	}
-	return false;
-}
-
-void FSequencerFilterBar::SetToVerticalLayout()
-{
-	USequencerSettings* const Settings = Sequencer.GetSequencerSettings();
-	if (!Settings)
-	{
-		return;
-	}
-
-	Settings->SetFilterBarLayout(EFilterBarLayout::Vertical);
-
-	StateChangedEvent.Broadcast(IsFilterBarVisible(), Settings->GetFilterBarLayout());
-}
-
-void FSequencerFilterBar::SetToHorizontalLayout()
-{
-	USequencerSettings* const Settings = Sequencer.GetSequencerSettings();
-	if (!Settings)
-	{
-		return;
-	}
-
-	Settings->SetFilterBarLayout(EFilterBarLayout::Horizontal);
-
-	StateChangedEvent.Broadcast(IsFilterBarVisible(), Settings->GetFilterBarLayout());
-}
-
-void FSequencerFilterBar::ToggleFilterBarLayout()
-{
-	if (IsFilterBarLayout(EFilterBarLayout::Horizontal))
-	{
-		SetToVerticalLayout();
-	}
-	else if (IsFilterBarLayout(EFilterBarLayout::Vertical))
-	{
-		SetToHorizontalLayout();
-	}
+	VisibilityChangedEvent.Broadcast(bNewVisibility);
 }
 
 bool FSequencerFilterBar::IsFilterSupported(const TSharedRef<FSequencerTrackFilter>& InFilter) const
diff --git a/Engine/Source/Editor/Sequencer/Private/Filters/SequencerFilterBar.h b/Engine/Source/Editor/Sequencer/Private/Filters/SequencerFilterBar.h
--- a/Engine/Source/Editor/Sequencer/Private/Filters/SequencerFilterBar.h
+++ b/Engine/Source/Editor/Sequencer/Private/Filters/SequencerFilterBar.h
@@ -234,15 +234,10 @@ public:
 	virtual bool IsFilterBarVisible() const override;
 	virtual void ToggleFilterBarVisibility() override;
 
-	virtual bool IsFilterBarLayout(const EFilterBarLayout InLayout) const override;
-	virtual void SetToVerticalLayout() override;
-	virtual void SetToHorizontalLayout() override;
-	virtual void ToggleFilterBarLayout() override;
-
 	bool IsFilterSupported(const TSharedRef<FSequencerTrackFilter>& InFilter) const;
 	bool IsFilterSupported(const FString& InFilterName) const;
 
-	virtual FOnFilterBarStateChanged& OnStateChanged() override { return StateChangedEvent; }
+	virtual FOnVisibilityChanged& OnVisibilityChanged() override { return VisibilityChangedEvent; }
 	virtual FSimpleMulticastDelegate& OnRequestUpdate() override { return RequestUpdateEvent; }
 	virtual FSimpleMulticastDelegate& OnTextFilterTextChanged() override { return OnTextFilterTextChangedEvent; }
 	virtual FMuteStateChangedDelegate& OnMuteFiltersChanged() override { return OnMuteFiltersChangedEvent; }
@@ -302,7 +297,7 @@ protected:
 	FSequencerFilterData FilterData;
 
 	FSequencerFiltersChanged FiltersChangedEvent;
-	FOnFilterBarStateChanged StateChangedEvent;
+	FOnVisibilityChanged VisibilityChangedEvent;
 	FSimpleMulticastDelegate RequestUpdateEvent;
 	FSimpleMulticastDelegate OnTextFilterTextChangedEvent;
 	FMuteStateChangedDelegate OnMuteFiltersChangedEvent;
diff --git a/Engine/Source/Editor/Sequencer/Private/Filters/SequencerFilterBarConfig.cpp b/Engine/Source/Editor/Sequencer/Private/Filters/SequencerFilterBarConfig.cpp
--- a/Engine/Source/Editor/Sequencer/Private/Filters/SequencerFilterBarConfig.cpp
+++ b/Engine/Source/Editor/Sequencer/Private/Filters/SequencerFilterBarConfig.cpp
@@ -104,12 +104,16 @@ bool FSequencerFilterBarConfig::RemoveCustomTextFilter(const FString& InFilterNa
 	return false;
 }
 
+PRAGMA_DISABLE_DEPRECATION_WARNINGS
 EFilterBarLayout FSequencerFilterBarConfig::GetFilterBarLayout() const
 {
-	return FilterBarLayout;
+	return FilterBarLayout_DEPRECATED;
 }
+PRAGMA_ENABLE_DEPRECATION_WARNINGS
 
+PRAGMA_DISABLE_DEPRECATION_WARNINGS
 void FSequencerFilterBarConfig::SetFilterBarLayout(const EFilterBarLayout InLayout)
 {
-	FilterBarLayout = InLayout;
+	FilterBarLayout_DEPRECATED = InLayout;
 }
+PRAGMA_ENABLE_DEPRECATION_WARNINGS
diff --git a/Engine/Source/Editor/Sequencer/Private/Filters/SequencerTrackFilterCommands.cpp b/Engine/Source/Editor/Sequencer/Private/Filters/SequencerTrackFilterCommands.cpp
--- a/Engine/Source/Editor/Sequencer/Private/Filters/SequencerTrackFilterCommands.cpp
+++ b/Engine/Source/Editor/Sequencer/Private/Filters/SequencerTrackFilterCommands.cpp
@@ -17,9 +17,6 @@ void FSequencerTrackFilterCommands::RegisterCommands()
 {
 	UI_COMMAND(ToggleFilterBarVisibility, "Filter Bar Visibility", "Toggle the visibility of the filter bar. The filter bar will only be displayed when there are enabled filters", EUserInterfaceActionType::ToggleButton, FInputChord(EModifierKey::Alt, EKeys::F));
 
-	UI_COMMAND(SetToVerticalLayout, "Vertical", "Swap to a vertical layout for the filter bar", EUserInterfaceActionType::RadioButton, FInputChord());
-	UI_COMMAND(SetToHorizontalLayout, "Horizontal", "Swap to a horizontal layout for the filter bar", EUserInterfaceActionType::RadioButton, FInputChord());
-
 	UI_COMMAND(ResetFilters, "Reset Filters", "Reset all enabled filters", EUserInterfaceActionType::Button, FInputChord());
 
 	UI_COMMAND(ToggleMuteFilters, "Mute Filters", "Mute or Unmute all active filters", EUserInterfaceActionType::ToggleButton, FInputChord());
diff --git a/Engine/Source/Editor/Sequencer/Private/Filters/SequencerTrackFilterCommands.h b/Engine/Source/Editor/Sequencer/Private/Filters/SequencerTrackFilterCommands.h
--- a/Engine/Source/Editor/Sequencer/Private/Filters/SequencerTrackFilterCommands.h
+++ b/Engine/Source/Editor/Sequencer/Private/Filters/SequencerTrackFilterCommands.h
@@ -21,9 +21,6 @@ public:
 
 	TSharedPtr<FUICommandInfo> ToggleFilterBarVisibility;
 
-	TSharedPtr<FUICommandInfo> SetToVerticalLayout;
-	TSharedPtr<FUICommandInfo> SetToHorizontalLayout;
-
 	TSharedPtr<FUICommandInfo> ResetFilters;
 
 	TSharedPtr<FUICommandInfo> ToggleMuteFilters;
diff --git a/Engine/Source/Editor/Sequencer/Private/Filters/Widgets/Linking/SFilterBarSwitcher.cpp b/Engine/Source/Editor/Sequencer/Private/Filters/Widgets/Linking/SFilterBarSwitcher.cpp
--- a/Engine/Source/Editor/Sequencer/Private/Filters/Widgets/Linking/SFilterBarSwitcher.cpp
+++ b/Engine/Source/Editor/Sequencer/Private/Filters/Widgets/Linking/SFilterBarSwitcher.cpp
@@ -20,13 +20,11 @@ void SFilterBarSwitcher::Construct(const FArguments& InArgs)
 		.LinkedContent()
 		[
 			SNew(SSequencerFilterBar, LinkedFilterViewModel->GetLinkedFilterBarImpl())
-			.FilterBarLayout(InArgs._FilterBarLayout)
 			.AddMetaData<FTagMetaData>(FTagMetaData(TEXT("SequencerTrackFilters")))
 		]
 		.InstancedContent()
 		[
 			SNew(SSequencerFilterBar, LinkedFilterViewModel->GetInstancedFilterBarImpl())
-			.FilterBarLayout(InArgs._FilterBarLayout)
 			.AddMetaData<FTagMetaData>(FTagMetaData(TEXT("SequencerTrackFilters")))
 		]
 	];
diff --git a/Engine/Source/Editor/Sequencer/Private/Filters/Widgets/Linking/SFilterBarSwitcher.h b/Engine/Source/Editor/Sequencer/Private/Filters/Widgets/Linking/SFilterBarSwitcher.h
--- a/Engine/Source/Editor/Sequencer/Private/Filters/Widgets/Linking/SFilterBarSwitcher.h
+++ b/Engine/Source/Editor/Sequencer/Private/Filters/Widgets/Linking/SFilterBarSwitcher.h
@@ -19,9 +19,6 @@ public:
 	SLATE_BEGIN_ARGS(SFilterBarSwitcher) {}
 		/** View-model determines the ELinkedFilterMode driving the displayed content. */
 		SLATE_ARGUMENT(TSharedPtr<FLinkedFilterViewModel>, LinkedFilterViewModel)
-
-		/** The layout that determines how the filters are laid out */
-		SLATE_ARGUMENT(EFilterBarLayout, FilterBarLayout)
 	SLATE_END_ARGS()
 	
 	void Construct(const FArguments& InArgs);
diff --git a/Engine/Source/Editor/Sequencer/Private/Filters/Widgets/SSequencerFilterBar.cpp b/Engine/Source/Editor/Sequencer/Private/Filters/Widgets/SSequencerFilterBar.cpp
--- a/Engine/Source/Editor/Sequencer/Private/Filters/Widgets/SSequencerFilterBar.cpp
+++ b/Engine/Source/Editor/Sequencer/Private/Filters/Widgets/SSequencerFilterBar.cpp
@@ -39,8 +39,6 @@ void SSequencerFilterBar::Construct(const FArguments& InArgs, const TSharedRef<F
 {
 	WeakFilterBar = InFilterBar;
 
-	FilterBarLayout = InArgs._FilterBarLayout;
-	bCanChangeOrientation = InArgs._CanChangeOrientation;
 	FilterPillStyle = InArgs._FilterPillStyle;
 
 	ContextMenu = MakeShared<FSequencerFilterBarContextMenu>();
@@ -51,32 +49,15 @@ void SSequencerFilterBar::Construct(const FArguments& InArgs, const TSharedRef<F
 
 	ChildSlot
 	[
-		SAssignNew(FilterBoxWidget, SWidgetSwitcher)
-		.WidgetIndex_Lambda([this]()
-			{
-				return (FilterBarLayout == EFilterBarLayout::Horizontal) ? 0 : 1;
-			})
-		+ SWidgetSwitcher::Slot()
-		.Padding(FMargin(0.f, 2.f, 0.f, 0.f))
+		SNew(SHorizontalBox)
+		+ SHorizontalBox::Slot()
 		[
-			SNew(SHorizontalBox)
-			+ SHorizontalBox::Slot()
-			[
-				HorizontalContainerWidget.ToSharedRef()
-			]
-			+ SHorizontalBox::Slot()
-			.AutoWidth()
-			[
-				HorizontalContainerWidget->CreateWrapButton()
-			]
+			HorizontalContainerWidget.ToSharedRef()
 		]
-		+ SWidgetSwitcher::Slot()
+		+ SHorizontalBox::Slot()
+		.AutoWidth()
 		[
-			SAssignNew(VerticalContainerWidget, SScrollBox)
-			.Visibility_Lambda([this]()
-				{
-				   return HasAnyFilterWidgets() ? EVisibility::Visible : EVisibility::Collapsed;
-				})
+			HorizontalContainerWidget->CreateWrapButton()
 		]
 	];
 	
@@ -117,29 +98,6 @@ TSharedPtr<FSequencerFilterBar> SSequencerFilterBar::GetFilterBar() const
 	return WeakFilterBar.Pin();
 }
 
-EFilterBarLayout SSequencerFilterBar::GetLayout() const
-{
-	return FilterBarLayout;
-}
-
-void SSequencerFilterBar::SetLayout(const EFilterBarLayout InFilterBarLayout)
-{
-	if (!bCanChangeOrientation)
-	{
-		return;
-	}
-
-	FilterBarLayout = InFilterBarLayout;
-
-	HorizontalContainerWidget->ClearChildren();
-	VerticalContainerWidget->ClearChildren();
-
-	for (const TPair<TSharedRef<FSequencerTrackFilter>, TSharedRef<SSequencerFilter>>& FilterPair : FilterWidgets)
-	{
-		AddWidgetToLayout(FilterPair.Value);
-	}
-}
-
 bool SSequencerFilterBar::HasAnyFilterWidgets() const
 {
 	return FilterWidgets.Num() > 0;
@@ -147,48 +105,17 @@ bool SSequencerFilterBar::HasAnyFilterWidgets() const
 
 void SSequencerFilterBar::AddWidgetToLayout(const TSharedRef<SWidget>& InWidget)
 {
-	FMargin SlotPadding = FMargin(1); // default editor-wide is FMargin(4, 2) for vertical only
-
-	if (FilterBarLayout == EFilterBarLayout::Horizontal)
-	{
-		switch (FilterPillStyle)
-		{
-		case EFilterPillStyle::Basic:
-			SlotPadding = FMargin(1); // default editor-wide is 2
-			break;
-		case EFilterPillStyle::Default:
-		default:
-			SlotPadding = FMargin(1); // default editor-wide is 3
-		}
-
-		HorizontalContainerWidget->AddSlot()
-			.AutoWidth()
-			.Padding(SlotPadding)
-			[
-				InWidget
-			];
-	}
-	else
-	{
-		VerticalContainerWidget->AddSlot()
-			.AutoSize()
-			.Padding(SlotPadding)
-			[
-				InWidget
-			];
-	}
+	HorizontalContainerWidget->AddSlot()
+		.AutoWidth()
+		.Padding(1)
+		[
+			InWidget
+		];
 }
 
 void SSequencerFilterBar::RemoveWidgetFromLayout(const TSharedRef<SWidget>& InWidget)
 {
-	if (FilterBarLayout == EFilterBarLayout::Horizontal)
-	{
-		HorizontalContainerWidget->RemoveSlot(InWidget);
-	}
-	else
-	{
-		VerticalContainerWidget->RemoveSlot(InWidget);
-	}
+	HorizontalContainerWidget->RemoveSlot(InWidget);
 }
 
 void SSequencerFilterBar::CreateAndAddFilterWidget(const TSharedRef<FSequencerTrackFilter>& InFilter)
@@ -430,15 +357,7 @@ void SSequencerFilterBar::OnFiltersChanged(const ESequencerFilterChange InChange
 void SSequencerFilterBar::OnMuteStateChanged(bool bNewIsMuted)
 {
 	const bool bIsEnabled = !bNewIsMuted;
-	if (HorizontalContainerWidget.IsValid())
-	{
-		HorizontalContainerWidget->SetEnabled(bIsEnabled);
-	}
-
-	if (VerticalContainerWidget.IsValid())
-	{
-		VerticalContainerWidget->SetEnabled(bIsEnabled);
-	}
+	HorizontalContainerWidget->SetEnabled(bIsEnabled);
 }
 
 void SSequencerFilterBar::CreateFilterWidgetsFromConfig()
diff --git a/Engine/Source/Editor/Sequencer/Private/Filters/Widgets/SSequencerFilterBar.h b/Engine/Source/Editor/Sequencer/Private/Filters/Widgets/SSequencerFilterBar.h
--- a/Engine/Source/Editor/Sequencer/Private/Filters/Widgets/SSequencerFilterBar.h
+++ b/Engine/Source/Editor/Sequencer/Private/Filters/Widgets/SSequencerFilterBar.h
@@ -20,17 +20,9 @@ class SSequencerFilterBar : public SCompoundWidget
 {
 public:
 	SLATE_BEGIN_ARGS(SSequencerFilterBar)
-		: _FilterBarLayout(EFilterBarLayout::Vertical)
-		, _CanChangeOrientation(true)
-		, _FilterPillStyle(EFilterPillStyle::Default)
+		: _FilterPillStyle(EFilterPillStyle::Default)
 		, _UseSectionsForCategories(true)
 	{}
-		/** The layout that determines how the filters are laid out */
-		SLATE_ARGUMENT(EFilterBarLayout, FilterBarLayout)
-
-		/** If true, allow dynamically changing the orientation and saving in the config */
-		SLATE_ARGUMENT(bool, CanChangeOrientation)
-
 		/** Determines how each individual filter pill looks like */
 		SLATE_ARGUMENT(EFilterPillStyle, FilterPillStyle)
 
@@ -50,9 +42,6 @@ public:
 
 	TSharedPtr<FSequencerFilterBar> GetFilterBar() const;
 
-	EFilterBarLayout GetLayout() const;
-	void SetLayout(const EFilterBarLayout InFilterBarLayout);
-
 	bool HasAnyFilterWidgets() const;
 
 protected:
@@ -94,12 +83,8 @@ protected:
 
 	TWeakPtr<FSequencerFilterBar> WeakFilterBar;
 
-	TSharedPtr<SWidgetSwitcher> FilterBoxWidget;
 	TSharedPtr<SFilterBarClippingHorizontalBox> HorizontalContainerWidget;
-	TSharedPtr<SScrollBox> VerticalContainerWidget;
 
-	EFilterBarLayout FilterBarLayout = EFilterBarLayout::Vertical;
-	bool bCanChangeOrientation = true;
 	EFilterPillStyle FilterPillStyle = EFilterPillStyle::Default;
 
 	TMap<TSharedRef<FSequencerTrackFilter>, TSharedRef<SSequencerFilter>> FilterWidgets;
diff --git a/Engine/Source/Editor/Sequencer/Private/MVVM/Views/CurveEditor/CurveEditorWidgetOwner.cpp b/Engine/Source/Editor/Sequencer/Private/MVVM/Views/CurveEditor/CurveEditorWidgetOwner.cpp
--- a/Engine/Source/Editor/Sequencer/Private/MVVM/Views/CurveEditor/CurveEditorWidgetOwner.cpp
+++ b/Engine/Source/Editor/Sequencer/Private/MVVM/Views/CurveEditor/CurveEditorWidgetOwner.cpp
@@ -53,7 +53,6 @@ static FCurveEditorWidgets MakePresenterWidgets(
 		[
 			SNew(SFilterBarSwitcher)
 			.LinkedFilterViewModel(InFilteringViewModel)
-			.FilterBarLayout(EFilterBarLayout::Horizontal)
 		];
 	
 	const TSharedRef<SWidgetSwitcher> Content = 
diff --git a/Engine/Source/Editor/Sequencer/Private/SSequencer.cpp b/Engine/Source/Editor/Sequencer/Private/SSequencer.cpp
--- a/Engine/Source/Editor/Sequencer/Private/SSequencer.cpp
+++ b/Engine/Source/Editor/Sequencer/Private/SSequencer.cpp
@@ -187,7 +187,6 @@ void SSequencer::Construct(const FArguments& InArgs, TSharedRef<FSequencer> InSe
 
 	SequencerPtr = InSequencer;
 	bIsActiveTimerRegistered = false;
-	bUserIsSelecting = false;
 	CachedClampRange = TRange<double>::Empty();
 	CachedViewRange = TRange<double>::Empty();
 	PendingFocus = FPendingWidgetFocus::MakeNoTextEdit();
@@ -365,7 +364,7 @@ void SSequencer::Construct(const FArguments& InArgs, TSharedRef<FSequencer> InSe
 	// TODO UE-363287
 	if (const TSharedPtr<FSequencerFilterBar> FilterBar = GetFilterBar())
 	{
-		FilterBar->OnStateChanged().AddSP(this, &SSequencer::OnFilterBarStateChanged);
+		FilterBar->OnVisibilityChanged().AddSP(this, &SSequencer::OnFilterBarVisibilityChanged);
 		FilterBar->OnFiltersChanged().AddSP(this, &SSequencer::OnTrackFiltersChanged);
 	}
 
@@ -635,49 +634,8 @@ void SSequencer::RebuildFilterBarContent()
 
 TSharedRef<SWidget> SSequencer::ConstructFilterBarContent()
 {
-	using namespace UE::Sequencer;
-	
 	RebuildSearchAndFilterRow();
-	
-	if (!IsFilterBarVisible() || GetFilterBarLayout() == EFilterBarLayout::Horizontal)
-	{
-		return OutlinerAndTrackAreaContent.ToSharedRef();
-	}
-
-	const TSharedRef<SFilterBarSwitcher> FilterBarWidgetRef = FilterBarWidget.ToSharedRef();
-
-	return SNew(SSplitter)
-		.Orientation(Orient_Horizontal)
-		.OnSplitterFinishedResizing_Lambda([this]()
-			{
-				GetSequencerSettings()->SaveConfig();
-			})
-
-		+ SSplitter::Slot()
-		.Value(TAttribute<float>::CreateLambda([this]()
-			{
-				const TSharedPtr<FSequencerFilterBar> FilterBar = GetFilterBar();
-				if (FilterBar.IsValid()
-					&& (FilterBar->HasEnabledFilter() || FilterBar->HasEnabledCustomTextFilters()))
-				{
-					return GetSequencerSettings()->GetLastFilterBarSizeCoefficient();
-				}
-				return 0.f;
-			}))
-		.OnSlotResized_Lambda([this](const float InNewCoefficient)
-			{
-				GetSequencerSettings()->SetLastFilterBarSizeCoefficient(InNewCoefficient, false);
-			})
-		[
-			SFilterBarClippingHorizontalBox::WrapVerticalListWithHeading(
-				FilterBarWidgetRef, FPointerEventHandler::CreateSP(FilterBarWidgetRef, &SFilterBarSwitcher::OnMouseButtonUp)
-				)
-		]
-		+ SSplitter::Slot()
-		.Value(0.94f)
-		[
-			OutlinerAndTrackAreaContent.ToSharedRef()
-		];
+	return OutlinerAndTrackAreaContent.ToSharedRef();
 }
 
 TSharedRef<SSplitter> SSequencer::ConstructOutlinerAndTrackArea(const FArguments& InArgs, const FTimeSliderArgs& InTimeSliderArgs)
@@ -4117,28 +4075,7 @@ void SSequencer::ToggleFilterBarVisibility()
 	RebuildFilterBarContent();
 }
 
-EFilterBarLayout SSequencer::GetFilterBarLayout() const
-{
-	USequencerSettings* const SequencerSettings = GetSequencerSettings();
-	if (ensure(SequencerSettings))
-	{
-		return SequencerSettings->GetFilterBarLayout();
-	}
-	return EFilterBarLayout::Vertical;
-}
-
-void SSequencer::SetFilterBarLayout(const EFilterBarLayout InLayout)
-{
-	USequencerSettings* const SequencerSettings = GetSequencerSettings();
-	if (ensure(SequencerSettings))
-	{
-		SequencerSettings->SetFilterBarLayout(InLayout);
-	}
-
-	RebuildFilterBarContent();
-}
-
-void SSequencer::OnFilterBarStateChanged(const bool bInIsVisible, const EFilterBarLayout InNewLayout)
+void SSequencer::OnFilterBarVisibilityChanged(const bool bInIsVisible)
 {
 	RebuildFilterBarContent();
 }
@@ -4176,11 +4113,10 @@ void SSequencer::RebuildSearchAndFilterRow()
 	const TSharedPtr<FSequencer> SequencerPin = SequencerPtr.Pin();
 	check(SequencerPin);
 	FilterBarWidget = SNew(SFilterBarSwitcher)
-		.FilterBarLayout(GetFilterBarLayout())
 		.AddMetaData<FTagMetaData>(FTagMetaData(TEXT("SequencerTrackFilters")))
 		.LinkedFilterViewModel(SequencerPin->GetViewModelImpl()->GetLinkedFilterViewModelImpl());
 
-	if (IsFilterBarVisible() && GetFilterBarLayout() == EFilterBarLayout::Horizontal)
+	if (IsFilterBarVisible())
 	{
 		SearchAndFilterRow->AddSlot()
 			.AutoHeight()
diff --git a/Engine/Source/Editor/Sequencer/Private/SSequencer.h b/Engine/Source/Editor/Sequencer/Private/SSequencer.h
--- a/Engine/Source/Editor/Sequencer/Private/SSequencer.h
+++ b/Engine/Source/Editor/Sequencer/Private/SSequencer.h
@@ -688,9 +688,6 @@ public:
 	bool IsFilterBarVisible() const;
 	void ToggleFilterBarVisibility();
 
-	EFilterBarLayout GetFilterBarLayout() const;
-	void SetFilterBarLayout(const EFilterBarLayout InLayout);
-
 	bool IsTransportControlsPoppedOut() const;
 	void SetTransportControlsPoppedOut(const bool bInPoppedOut);
 	void ToggleTransportControlsPoppedOut();
@@ -708,7 +705,7 @@ private:
 
 	void OnSidebarStateChanged(const FSidebarState& InNewState);
 
-	void OnFilterBarStateChanged(const bool bInIsVisible, const EFilterBarLayout InNewLayout);
+	void OnFilterBarVisibilityChanged(bool bInIsVisible);
 
 	void OnTrackFiltersChanged(const ESequencerFilterChange InChangeType, const TSharedRef<FSequencerTrackFilter>& InFilter);
 
@@ -801,14 +798,9 @@ private:
 	/** List of registered outliner columns with their visibility states */
 	TArray<FSequencerOutlinerColumnVisibility> OutlinerColumnVisibilities;
 
-	TSharedPtr<class SSequencerSplitterOverlay> TreeViewSplitter;
-
 	/** Whether the active timer is currently registered */
 	bool bIsActiveTimerRegistered;
 
-	/** Whether the user is selecting. Ignore selection changes from the level when the user is selecting. */
-	bool bUserIsSelecting;
-
 	/** Default initialized in the view params to a lambda that gives us the standard speeds */
 	ISequencer::FOnGetPlaybackSpeeds OnGetPlaybackSpeeds;
 	
@@ -901,10 +893,9 @@ private:
 
 	TSharedPtr<FSequencerSelectionDrawer> SelectionDetailsDrawer;
 
-	/** Displays the filter pills. Placed at different UI locations depending on the layout mode. */
+	/** Displays the filter pills. */
 	TSharedPtr<UE::Sequencer::SFilterBarSwitcher> FilterBarWidget;
-	
-	/** Contains the filters widgets plus optionally the filter bar when in horizontal layout mode. */
+	/** Contains the filters widgets plus the filter bar. */
 	TSharedPtr<SVerticalBox> SearchAndFilterRow;
 	/** The widgets displayed above the tree view. */
 	TSharedPtr<UE::Sequencer::SSequencerSearchAndFilterRow> FilterRow;
diff --git a/Engine/Source/Editor/Sequencer/Private/SequencerSettings.cpp b/Engine/Source/Editor/Sequencer/Private/SequencerSettings.cpp
--- a/Engine/Source/Editor/Sequencer/Private/SequencerSettings.cpp
+++ b/Engine/Source/Editor/Sequencer/Private/SequencerSettings.cpp
@@ -1301,16 +1301,20 @@ void USequencerSettings::SetFilterBarVisible(const bool bInVisible)
 	SaveConfig();
 }
 
+PRAGMA_DISABLE_DEPRECATION_WARNINGS
 EFilterBarLayout USequencerSettings::GetFilterBarLayout() const
 {
-	return LastFilterBarLayout;
+	return LastFilterBarLayout_DEPRECATED;
 }
+PRAGMA_ENABLE_DEPRECATION_WARNINGS
 
+PRAGMA_DISABLE_DEPRECATION_WARNINGS
 void USequencerSettings::SetFilterBarLayout(const EFilterBarLayout InLayout)
 {
-	LastFilterBarLayout = InLayout;
+	LastFilterBarLayout_DEPRECATED = InLayout;
 	SaveConfig();
 }
+PRAGMA_ENABLE_DEPRECATION_WARNINGS
 
 float USequencerSettings::GetLastFilterBarSizeCoefficient() const
 {
diff --git a/Engine/Source/Editor/Sequencer/Public/Filters/ISequencerFilterBar.h b/Engine/Source/Editor/Sequencer/Public/Filters/ISequencerFilterBar.h
--- a/Engine/Source/Editor/Sequencer/Public/Filters/ISequencerFilterBar.h
+++ b/Engine/Source/Editor/Sequencer/Public/Filters/ISequencerFilterBar.h
@@ -19,7 +19,10 @@ struct FSequencerFilterSuggestion;
 class ISequencerFilterBar : public TSharedFromThis<ISequencerFilterBar>
 {
 public:
+	
+	// This type is deprecated. However, DECLARE_EVENT_TwoParams does not support UE_DEPRECATED macro in Linux.
 	DECLARE_EVENT_TwoParams(ISequencerFilterBar, FOnFilterBarStateChanged, const bool /*InIsVisible*/, const EFilterBarLayout /*InNewLayout*/);
+	DECLARE_MULTICAST_DELEGATE_OneParam(FOnVisibilityChanged, bool /*bIsVisible*/)
 	DECLARE_MULTICAST_DELEGATE_OneParam(FMuteStateChangedDelegate, bool /* bNewIsMuted */);
 	
 	virtual ~ISequencerFilterBar() = default;
@@ -87,7 +90,13 @@ public:
 	virtual void CreateNewTextFilter() = 0;
 
 	/** Event called when the visibility or layout state has changed */
-	virtual FOnFilterBarStateChanged& OnStateChanged() = 0;
+	UE_DEPRECATED(5.8, "Use OnVisibilityChanged instead")
+	PRAGMA_DISABLE_DEPRECATION_WARNINGS
+	virtual FOnFilterBarStateChanged& OnStateChanged() { return OnFilterBarStateChanged_DEPRECATED; }
+	PRAGMA_ENABLE_DEPRECATION_WARNINGS
+	
+	/** Event called when the visibility has changed. */
+	virtual FOnVisibilityChanged& OnVisibilityChanged() = 0;
 
 	/** Event called to request that the subscribe update its view since the filters state has changed */
 	virtual FSimpleMulticastDelegate& OnRequestUpdate() = 0;
@@ -103,10 +112,14 @@ public:
 	virtual bool IsFilterBarVisible() const = 0;
 	virtual void ToggleFilterBarVisibility() = 0;
 
-	virtual bool IsFilterBarLayout(const EFilterBarLayout InLayout) const = 0;
-	virtual void SetToVerticalLayout() = 0;
-	virtual void SetToHorizontalLayout() = 0;
-	virtual void ToggleFilterBarLayout() = 0;
+	UE_DEPRECATED(5.8, "This function was removed")
+	virtual bool IsFilterBarLayout(const EFilterBarLayout InLayout) const { return false; }
+	UE_DEPRECATED(5.8, "This function was removed")
+	virtual void SetToVerticalLayout() {}
+	UE_DEPRECATED(5.8, "This function was removed")
+	virtual void SetToHorizontalLayout() {}
+	UE_DEPRECATED(5.8, "This function was removed")
+	virtual void ToggleFilterBarLayout() {}
 
 	/** @return True if a common filter, specified by InFilterName, passes its filter. Used to support default text
 	 * expressions for common filters without needing to create an actual text expression class for each common filter. */
@@ -141,6 +154,12 @@ public:
 protected:
 	/** Internal implementation to bypass template limitations in virtual functions */
 	virtual void GetCommonFiltersImpl(TArray<TSharedPtr<void>>& OutTypeErasedFilters) const = 0;
+	
+private:
+	
+	PRAGMA_DISABLE_DEPRECATION_WARNINGS
+	FOnFilterBarStateChanged OnFilterBarStateChanged_DEPRECATED;
+	PRAGMA_ENABLE_DEPRECATION_WARNINGS
 };
 
 #include "Filters/ISequencerFilterBar.inl"
diff --git a/Engine/Source/Editor/Sequencer/Public/Filters/SequencerFilterBarConfig.h b/Engine/Source/Editor/Sequencer/Public/Filters/SequencerFilterBarConfig.h
--- a/Engine/Source/Editor/Sequencer/Public/Filters/SequencerFilterBarConfig.h
+++ b/Engine/Source/Editor/Sequencer/Public/Filters/SequencerFilterBarConfig.h
@@ -29,6 +29,14 @@ struct FSequencerFilterBarConfig
 
 public:
 	/** Common Filters */
+	
+	PRAGMA_DISABLE_DEPRECATION_WARNINGS
+	FSequencerFilterBarConfig() = default;
+	FSequencerFilterBarConfig(const FSequencerFilterBarConfig&) = default;
+	FSequencerFilterBarConfig(FSequencerFilterBarConfig&& Other) = default;
+	FSequencerFilterBarConfig& operator=(const FSequencerFilterBarConfig& Other) = default;
+	FSequencerFilterBarConfig& operator=(FSequencerFilterBarConfig&& Other) = default;
+	PRAGMA_ENABLE_DEPRECATION_WARNINGS
 
 	SEQUENCER_API bool IsFilterEnabled(const FString& InFilterName) const;
 	SEQUENCER_API bool SetFilterEnabled(const FString& InFilterName, const bool bInActive);
@@ -47,8 +55,9 @@ public:
 	SEQUENCER_API bool RemoveCustomTextFilter(const FString& InFilterName);
 
 	/** Filter Bar Layout */
-
+	UE_DEPRECATED(5.8, "Sequencer now uses horizontal filter layout.")
 	SEQUENCER_API EFilterBarLayout GetFilterBarLayout() const;
+	UE_DEPRECATED(5.8, "Sequencer now uses horizontal filter layout.")
 	SEQUENCER_API void SetFilterBarLayout(const EFilterBarLayout InLayout);
 
 protected:
@@ -61,6 +70,7 @@ protected:
 	TArray<FCustomTextFilterData> CustomTextFilters;
 
 	/** The layout style for the filter bar widget */
+	UE_DEPRECATED(5.8, "Sequencer now uses horizontal filter layout.")
 	UPROPERTY()
-	EFilterBarLayout FilterBarLayout = EFilterBarLayout::Vertical;
+	EFilterBarLayout FilterBarLayout_DEPRECATED = EFilterBarLayout::Vertical;
 };
diff --git a/Engine/Source/Editor/Sequencer/Public/SequencerSettings.h b/Engine/Source/Editor/Sequencer/Public/SequencerSettings.h
--- a/Engine/Source/Editor/Sequencer/Public/SequencerSettings.h
+++ b/Engine/Source/Editor/Sequencer/Public/SequencerSettings.h
@@ -591,7 +591,9 @@ public:
 	UE_API bool IsFilterBarVisible() const;
 	UE_API void SetFilterBarVisible(const bool bInVisible);
 
+	UE_DEPRECATED(5.8, "Sequencer now only uses horizontal layout")
 	UE_API EFilterBarLayout GetFilterBarLayout() const;
+	UE_DEPRECATED(5.8, "Sequencer now only uses horizontal layout")
 	UE_API void SetFilterBarLayout(const EFilterBarLayout InLayout);
 
 	UE_API float GetLastFilterBarSizeCoefficient() const;
@@ -970,8 +972,9 @@ protected:
 	bool bFilterBarVisible;
 
 	/** Last saved layout orientation of the filter bar to restore after closed */
-	UPROPERTY(config, EditAnywhere, Category = Filtering)
-	EFilterBarLayout LastFilterBarLayout;
+	UE_DEPRECATED(5.8, "Sequencer now only uses horizontal layout")
+	UPROPERTY(config)
+	EFilterBarLayout LastFilterBarLayout_DEPRECATED;
 
 	/** Last saved size of the filter bar to restore after closed */
 	UPROPERTY(config)
