Sequencer: Anim Mixer and Decoration MVVM refactor
* Mixer now always shows layers and section outliners even in single layer case and single section case
* Decorations/Modifiers no longer show a 'Modifiers' row to group the different modifiers- now each modifiers get their own row using the DecorationOutlinerModel and handle both section provider and channel provider cases

#Synced-CL 51336034

diff --git a/Engine/Plugins/Experimental/MovieSceneAnimMixer/Source/MovieSceneAnimMixerEditor/Private/MVVM/AnimationMixerLayerModel.cpp b/Engine/Plugins/Experimental/MovieSceneAnimMixer/Source/MovieSceneAnimMixerEditor/Private/MVVM/AnimationMixerLayerModel.cpp
--- a/Engine/Plugins/Experimental/MovieSceneAnimMixer/Source/MovieSceneAnimMixerEditor/Private/MVVM/AnimationMixerLayerModel.cpp
+++ b/Engine/Plugins/Experimental/MovieSceneAnimMixer/Source/MovieSceneAnimMixerEditor/Private/MVVM/AnimationMixerLayerModel.cpp
@@ -227,7 +227,7 @@ void FAnimationMixerLayerModel::RefreshLayout(bool bChildrenNeedLayout)
 			TArray<UMovieSceneSection*> AllSections = ChildTrack->GetAllSections();
 			// Query saved expansion state directly - we can't use IsExpanded() here because
 			// it checks for children, but we're in the process of building them
-			bool bCreateOutliners = GetSavedExpansionState() && AllSections.Num() > 1;
+			bool bCreateOutliners = GetSavedExpansionState();
 
 			// Only recycle top level channels if children need layout
 			if (bChildrenNeedLayout)
@@ -302,7 +302,7 @@ void FAnimationMixerLayerModel::RefreshLayout(bool bChildrenNeedLayout)
 
 			// Query saved expansion state directly - we can't use IsExpanded() here because
 			// it checks for children, but we're in the process of building them
-			bool bCreateOutliners = GetSavedExpansionState() && (LayerSections.Num() > 1 || SequencerHelpers::HasOutlinerVisibleDecorations(Layer));
+			bool bCreateOutliners = GetSavedExpansionState();
 
 			// Only recycle top level channels if children need layout
 			if (bChildrenNeedLayout)
@@ -329,17 +329,11 @@ void FAnimationMixerLayerModel::RefreshLayout(bool bChildrenNeedLayout)
 
 			PreviousLayoutNumSections = LayerSections.Num();
 
-			// Handle layer and section decorations
+			// Handle layer decorations. Section decorations are handled by
+			// RefreshDecorations() on each section outliner (called from LayoutSections).
 			if (FDecorationModelStorageExtension* DecorationModelStorage = SequenceModel->CastDynamic<FDecorationModelStorageExtension>())
 			{
-				DecorationModelStorage->SyncDecorationModel(Layer, ParentTrack, OutlinerChildren);
-
-				// Section decorations for single-section case (when section outliners aren't created).
-				// Multi-section case: SyncDecorationModel's section outliner iteration handles stale cleanup.
-				if (LayerSections.Num() == 1)
-				{
-					DecorationModelStorage->SyncDecorationModel(LayerSections[0], ParentTrack, OutlinerChildren);
-				}
+				DecorationModelStorage->SyncDecorationModels(Layer, ParentTrack, OutlinerChildren);
 			}
 		}
 	}
diff --git a/Engine/Plugins/Experimental/MovieSceneAnimMixer/Source/MovieSceneAnimMixerEditor/Private/MVVM/AnimationMixerLayoutHelpers.cpp b/Engine/Plugins/Experimental/MovieSceneAnimMixer/Source/MovieSceneAnimMixerEditor/Private/MVVM/AnimationMixerLayoutHelpers.cpp
--- a/Engine/Plugins/Experimental/MovieSceneAnimMixer/Source/MovieSceneAnimMixerEditor/Private/MVVM/AnimationMixerLayoutHelpers.cpp
+++ b/Engine/Plugins/Experimental/MovieSceneAnimMixer/Source/MovieSceneAnimMixerEditor/Private/MVVM/AnimationMixerLayoutHelpers.cpp
@@ -220,32 +220,29 @@ void LayoutSections(
 
 		TArray<TSharedPtr<FViewModel>> SectionRoots;
 
-		// For multiple sections with outliner creation enabled, create section outliner models
-		// For single section or outliner creation disabled, add sections directly
-		const int32 NumSections = SectionModels.Num();
-		if (bCreateSectionOutliners && NumSections > 1)
+		// When outliner creation is enabled, create section outliner models for each section.
+		// Otherwise add sections directly to the track area.
+		if (bCreateSectionOutliners)
 		{
-			// Multiple sections: Create section outliner models
 			SectionRoots = CreateSectionOutliners(SectionModels, OutlinerChildren, RecycledModels.GetChildren(), Owner);
 		}
 		else
 		{
-			// Single section or no outliners: Add directly to section children
 			SectionRoots = AddSectionsDirectly(SectionModels, SectionChildren, Owner);
 		}
 
 		// Run layout builder on each section
 		RunLayoutBuilders(SectionModels, SectionRoots);
+	}
 
-		// Refresh decorations on section outliners after layout builders complete.
-		// Layout builders recycle the outliner child list, so decorations must be
-		// added after that process finishes.
-		for (TSharedPtr<FSectionOutlinerModel> SectionOutliner : OutlinerChildren.IterateSubList<FSectionOutlinerModel>())
+	// Refresh decorations on section outliners. This runs outside the rebuild block
+	// so that decoration state changes (e.g. channel proxy dirty) are detected even
+	// when the layout itself didn't change.
+	for (TSharedPtr<FSectionOutlinerModel> SectionOutliner : OutlinerChildren.IterateSubList<FSectionOutlinerModel>())
+	{
+		if (SectionOutliner)
 		{
-			if (SectionOutliner)
-			{
-				SectionOutliner->RefreshDecorations();
-			}
+			SectionOutliner->RefreshDecorations();
 		}
 	}
 }
diff --git a/Engine/Plugins/Experimental/MovieSceneAnimMixer/Source/MovieSceneAnimMixerEditor/Private/MVVM/AnimationMixerTrackModel.cpp b/Engine/Plugins/Experimental/MovieSceneAnimMixer/Source/MovieSceneAnimMixerEditor/Private/MVVM/AnimationMixerTrackModel.cpp
--- a/Engine/Plugins/Experimental/MovieSceneAnimMixer/Source/MovieSceneAnimMixerEditor/Private/MVVM/AnimationMixerTrackModel.cpp
+++ b/Engine/Plugins/Experimental/MovieSceneAnimMixer/Source/MovieSceneAnimMixerEditor/Private/MVVM/AnimationMixerTrackModel.cpp
@@ -16,10 +16,8 @@
 #include "MVVM/Extensions/ITrackAreaExtension.h"
 #include "MVVM/Extensions/ITrackExtension.h"
 #include "MVVM/Extensions/IRecyclableExtension.h"
-#include "MVVM/Extensions/IObjectBindingExtension.h"
 #include "MVVM/TrackModelStorageExtension.h"
 #include "MVVM/TrackRowModelStorageExtension.h"
-#include "MVVM/SectionModelStorageExtension.h"
 #include "MVVM/ViewModels/ViewModelIterators.h"
 #include "MVVM/SharedViewModelData.h"
 #include "MVVM/ViewModels/ViewModelHierarchy.h"
@@ -66,134 +64,61 @@ void FAnimationMixerTrackModel::ForceUpdate()
 
 	const TArray<TObjectPtr<UMovieSceneAnimationMixerLayer>>& Layers = MixerTrack->GetLayers();
 
-	// Collapse logic: If there's only one layer, show its contents directly without the layer wrapper
-	if (Layers.Num() == 1)
+	// Always create layer models for all layers
+	TSharedPtr<FAnimationMixerLayerModel> LastLayerModel;
+
+	for (UMovieSceneAnimationMixerLayer* Layer : Layers)
 	{
-		UMovieSceneAnimationMixerLayer* Layer = Layers[0];
 		if (!Layer)
 		{
-			return;
+			continue;
 		}
 
-		if (Layer->HasChildTrack() || Layer->IsEmpty())
-		{
-			// Child track case - use a layer model since it needs the ITrackExtension interface
-			TSharedPtr<FAnimationMixerLayerModel> LayerModel = MakeShared<FAnimationMixerLayerModel>(MixerTrack, Layer, Layer->GetChildTrack());
-
-			// Add to outliner so it can find ancestors
-			OutlinerChildren.AddChild(LayerModel);
+		// Find existing layer model or create new one
+		TSharedPtr<FAnimationMixerLayerModel> LayerModel;
 
-			// Initialize and refresh layout
-			LayerModel->OnConstruct();
-			LayerModel->RefreshLayout(true);
-		}
-		else
+		// Try to find existing model in recycled children
+		for (TViewModelPtr<FAnimationMixerLayerModel> ExistingModel : RecycledOutliner.GetChildren().IterateSubList<FAnimationMixerLayerModel>().ToArray())
 		{
-			// For single layer, create section models directly without layer wrapper
-			// This follows the same pattern as FTrackModel for single-row tracks
-
-			FSectionModelStorageExtension* SectionModelStorage = SequenceModel->CastDynamic<FSectionModelStorageExtension>();
-			FDecorationModelStorageExtension* DecorationModelStorage = SequenceModel->CastDynamic<FDecorationModelStorageExtension>();
-			TSharedPtr<ISequencerTrackEditor> TrackEditorPtr = GetTrackEditor();
-			if (SectionModelStorage && TrackEditorPtr)
+			if (ExistingModel && ExistingModel->GetLayer() == Layer)
 			{
-				FGuid ObjectBinding;
-				if (TSharedPtr<IObjectBindingExtension> ObjectBindingExtension = FindAncestorOfType<IObjectBindingExtension>())
-				{
-					ObjectBinding = ObjectBindingExtension->GetObjectGuid();
-				}
-
-				TArray<UMovieSceneSection*> LayerSections = Layer->GetSections();
-				bool bCreateOutliners = LayerSections.Num() > 1;
-				bool bChildrenNeedLayout = true;
-
-				FViewModelChildren SectionChildren = GetSectionModels();
-
-				AnimationMixerHelpers::LayoutSections(
-					LayerSections,
-					MixerTrack,
-					AsShared(),
-					OutlinerChildren,
-					SectionChildren,
-					SectionModelStorage,
-					TrackEditorPtr,
-					ObjectBinding,
-					bCreateOutliners,
-					bChildrenNeedLayout);
-
-				// Handle layer and section decorations
-				if (DecorationModelStorage)
-				{
-					DecorationModelStorage->SyncDecorationModel(Layer, MixerTrack, OutlinerChildren);
-
-					// Section decorations for single-section case (when section outliners aren't created).
-					// Multi-section case: SyncDecorationModel's section outliner iteration handles stale cleanup.
-					if (!bCreateOutliners && LayerSections.Num() == 1)
-					{
-						DecorationModelStorage->SyncDecorationModel(LayerSections[0], MixerTrack, OutlinerChildren);
-					}
-				}
+				LayerModel = ExistingModel;
+				break;
 			}
 		}
-	}
-	// Multiple layers: Show layer hierarchy
-	else if (Layers.Num() > 1)
-	{
-		TSharedPtr<FAnimationMixerLayerModel> LastLayerModel;
 
-		for (UMovieSceneAnimationMixerLayer* Layer : Layers)
+		// Create new layer model if not found
+		bool bNeedsLayout = false;
+		if (!LayerModel)
 		{
-			if (!Layer)
+			if (Layer->HasChildTrack())
 			{
-				continue;
-			}
-
-			// Find existing layer model or create new one
-			TSharedPtr<FAnimationMixerLayerModel> LayerModel;
-
-			// Try to find existing model in recycled children
-			for (TViewModelPtr<FAnimationMixerLayerModel> ExistingModel : RecycledOutliner.GetChildren().IterateSubList<FAnimationMixerLayerModel>().ToArray())
-			{
-				if (ExistingModel && ExistingModel->GetLayer() == Layer)
-				{
-					LayerModel = ExistingModel;
-					break;
-				}
+				LayerModel = MakeShared<FAnimationMixerLayerModel>(MixerTrack, Layer, Layer->GetChildTrack());
 			}
-
-			// Create new layer model if not found
-			bool bNeedsLayout = false;
-			if (!LayerModel)
+			else
 			{
-				if (Layer->HasChildTrack())
-				{
-					LayerModel = MakeShared<FAnimationMixerLayerModel>(MixerTrack, Layer, Layer->GetChildTrack());
-				}
-				else
-				{
-					LayerModel = MakeShared<FAnimationMixerLayerModel>(MixerTrack, Layer);
-				}
-				bNeedsLayout = true;
+				LayerModel = MakeShared<FAnimationMixerLayerModel>(MixerTrack, Layer);
 			}
+			bNeedsLayout = true;
+		}
 
-			// Add to outliner
-			OutlinerChildren.InsertChild(LayerModel, LastLayerModel);
-			LastLayerModel = LayerModel;
+		// Add to outliner
+		OutlinerChildren.InsertChild(LayerModel, LastLayerModel);
+		LastLayerModel = LayerModel;
 
-			// Initialize and refresh layout
-			if (bNeedsLayout)
-			{
-				LayerModel->OnConstruct();
-			}
-			LayerModel->RefreshLayout(bNeedsLayout);
+		// Initialize and refresh layout
+		if (bNeedsLayout)
+		{
+			LayerModel->OnConstruct();
 		}
+		LayerModel->RefreshLayout(bNeedsLayout);
 	}
 
 	// Handle track-level decorations
 	FDecorationModelStorageExtension* DecorationModelStorage = SequenceModel->CastDynamic<FDecorationModelStorageExtension>();
 	if (DecorationModelStorage)
 	{
-		DecorationModelStorage->SyncDecorationModel(MixerTrack, MixerTrack, OutlinerChildren);
+		DecorationModelStorage->SyncDecorationModels(MixerTrack, MixerTrack, OutlinerChildren);
 	}
 }
 
@@ -411,35 +336,6 @@ TOptional<FAnimMixerDropPreview> FAnimationMixerTrackModel::ComputeDropPreview(c
 	const bool bCtrlHeld = FSlateApplication::Get().GetModifierKeys().IsControlDown();
 	const int32 LastLayerIndex = Layers.Num() - 1;
 
-	// Handle single-layer collapsed view
-	if (Layers.Num() == 1 && bIsInDrag)
-	{
-		FVirtualGeometry TrackGeometry;
-		if (IGeometryExtension* GeometryExt = const_cast<FAnimationMixerTrackModel*>(this)->CastThis<IGeometryExtension>())
-		{
-			TrackGeometry = GeometryExt->GetVirtualGeometry();
-		}
-
-		float TrackTop = TrackGeometry.Top;
-		float TrackBottom = TrackGeometry.NestedBottom;
-
-		if (Context.VirtualMousePos.Y < TrackTop || Context.LocalMousePos.Y <= 0)
-		{
-			// Above track - always allowed (new layer at top)
-			return FAnimMixerDropPreview{ FAnimMixerDropPreview::EZone::AboveLayer, 0, TrackTop, TrackBottom };
-		}
-		else if (Context.VirtualMousePos.Y > TrackBottom)
-		{
-			// Below track - always allowed (new layer at bottom)
-			return FAnimMixerDropPreview{ FAnimMixerDropPreview::EZone::BelowLayer, 0, TrackTop, TrackBottom };
-		}
-		else
-		{
-			// Within track - drop onto layer 0
-			return FAnimMixerDropPreview{ FAnimMixerDropPreview::EZone::OntoLayer, 0, TrackTop, TrackBottom, DraggedSectionRange };
-		}
-	}
-
 	// Iterate through layer models to find hover target
 	int32 LastRowIndex = -1;
 	for (TSharedPtr<FViewModel> ChildNode : const_cast<FAnimationMixerTrackModel*>(this)->GetChildren())
diff --git a/Engine/Source/Editor/Sequencer/Private/MVVM/DecorationModelStorageExtension.cpp b/Engine/Source/Editor/Sequencer/Private/MVVM/DecorationModelStorageExtension.cpp
--- a/Engine/Source/Editor/Sequencer/Private/MVVM/DecorationModelStorageExtension.cpp
+++ b/Engine/Source/Editor/Sequencer/Private/MVVM/DecorationModelStorageExtension.cpp
@@ -5,7 +5,9 @@
 #include "MVVM/ViewModels/DecorationOutlinerModel.h"
 #include "MVVM/ViewModels/SectionOutlinerModel.h"
 #include "MVVM/ViewModels/ViewModelHierarchy.h"
-#include "SequencerCommonHelpers.h"
+#include "Decorations/IMovieSceneChannelDecoration.h"
+#include "Decorations/IMovieSceneSectionProviderDecoration.h"
+#include "Decorations/MovieSceneDecorationContainer.h"
 
 namespace UE::Sequencer
 {
@@ -16,91 +18,122 @@ FDecorationModelStorageExtension::FDecorationModelStorageExtension()
 
 void FDecorationModelStorageExtension::OnReinitialize()
 {
-	for (auto It = ContainerToModel.CreateIterator(); It; ++It)
+	for (auto It = DecorationToModel.CreateIterator(); It; ++It)
 	{
 		if (It.Key().ResolveObjectPtr() == nullptr || It.Value().Pin().Get() == nullptr)
 		{
 			It.RemoveCurrent();
 		}
 	}
-	ContainerToModel.Compact();
+	DecorationToModel.Compact();
 }
 
 TSharedPtr<FDecorationOutlinerModel> FDecorationModelStorageExtension::FindOrCreateModel(
+	UObject* Decoration,
 	UMovieSceneDecorationContainerObject* Container,
 	UMovieSceneTrack* ParentTrack,
 	bool& bOutNewlyCreated)
 {
 	bOutNewlyCreated = false;
 
-	if (!Container || !ParentTrack)
+	if (!Decoration || !Container || !ParentTrack)
 	{
 		return nullptr;
 	}
 
-	FObjectKey ContainerKey(Container);
-	if (TSharedPtr<FDecorationOutlinerModel> Existing = ContainerToModel.FindRef(ContainerKey).Pin())
+	FObjectKey DecorationKey(Decoration);
+	if (TSharedPtr<FDecorationOutlinerModel> Existing = DecorationToModel.FindRef(DecorationKey).Pin())
 	{
 		return Existing;
 	}
 
-	TSharedPtr<FDecorationOutlinerModel> NewModel = MakeShared<FDecorationOutlinerModel>(ParentTrack, Container);
-	ContainerToModel.Add(ContainerKey, NewModel);
+	TSharedPtr<FDecorationOutlinerModel> NewModel = MakeShared<FDecorationOutlinerModel>(ParentTrack, Container, Decoration);
+	DecorationToModel.Add(DecorationKey, NewModel);
 	bOutNewlyCreated = true;
 	return NewModel;
 }
 
 TSharedPtr<FDecorationOutlinerModel> FDecorationModelStorageExtension::FindModel(
-	const UMovieSceneDecorationContainerObject* Container) const
+	const UObject* Decoration) const
 {
-	FObjectKey ContainerKey(Container);
-	return ContainerToModel.FindRef(ContainerKey).Pin();
+	FObjectKey DecorationKey(Decoration);
+	return DecorationToModel.FindRef(DecorationKey).Pin();
 }
 
-TSharedPtr<FDecorationOutlinerModel> FDecorationModelStorageExtension::SyncDecorationModel(
+void FDecorationModelStorageExtension::SyncDecorationModels(
 	UMovieSceneDecorationContainerObject* Container,
 	UMovieSceneTrack* ParentTrack,
 	FViewModelChildren& OutlinerChildren)
 {
 	if (!Container || !ParentTrack)
 	{
-		return nullptr;
+		return;
 	}
 
-	TSharedPtr<FDecorationOutlinerModel> Result;
+	// Build the set of currently visible decorations on this container
+	TSet<UObject*> VisibleDecorations;
+	for (const TObjectPtr<UObject>& Decoration : Container->GetDecorations())
+	{
+		if (!Decoration)
+		{
+			continue;
+		}
 
-	if (SequencerHelpers::HasOutlinerVisibleDecorations(Container))
+		if (Cast<IMovieSceneChannelDecoration>(Decoration.Get()) ||
+			Cast<IMovieSceneSectionProviderDecoration>(Decoration.Get()))
+		{
+			VisibleDecorations.Add(Decoration.Get());
+		}
+	}
+
+	// Create or update models for each visible decoration
+	for (UObject* Decoration : VisibleDecorations)
 	{
 		bool bNewlyCreated = false;
-		TSharedPtr<FDecorationOutlinerModel> Model = FindOrCreateModel(Container, ParentTrack, bNewlyCreated);
+		TSharedPtr<FDecorationOutlinerModel> Model = FindOrCreateModel(Decoration, Container, ParentTrack, bNewlyCreated);
 		if (Model)
 		{
 			OutlinerChildren.AddChild(Model);
 			// New models get OnConstruct via the framework (SetSharedData). For existing models,
-			// only reconstruct if the container's decorations actually changed, otherwise keep
-			// existing children intact to avoid flickering during layout-only rebuilds (e.g. drags).
+			// only reconstruct if the decoration's state actually changed.
 			if (!bNewlyCreated && Model->NeedsReconstruct())
 			{
 				Model->OnConstruct();
 			}
 		}
-		Result = Model;
 	}
-	else
+
+	// Collect stale decoration models before removing (can't modify list during iteration)
+	TArray<TSharedPtr<FDecorationOutlinerModel>> StaleModels;
+	for (TSharedPtr<FDecorationOutlinerModel> DecorationModel : OutlinerChildren.IterateSubList<FDecorationOutlinerModel>())
 	{
-		// Container has no visible decorations - remove any stale model
-		if (TSharedPtr<FDecorationOutlinerModel> StaleModel = FindModel(Container))
+		if (!DecorationModel)
 		{
-			StaleModel->RemoveFromParent();
+			continue;
+		}
+
+		// Only remove models that belong to this container
+		if (DecorationModel->GetDecorationContainer().Get() != Container)
+		{
+			continue;
+		}
+
+		UObject* Decoration = DecorationModel->GetDecoration();
+		if (!Decoration || !VisibleDecorations.Contains(Decoration))
+		{
+			StaleModels.Add(DecorationModel);
 		}
 	}
 
+	for (const TSharedPtr<FDecorationOutlinerModel>& StaleModel : StaleModels)
+	{
+		StaleModel->RemoveFromParent();
+	}
+
 	// Clean up stale decoration models on section outliners. When decorations are
-	// removed from a section (e.g. via Delete), the event fires before the decoration
-	// is actually removed from the container, so the event-driven cleanup can't detect
-	// the change. This deferred check catches it during the subsequent rebuild.
-	// Only does removal - creation/adoption is handled by RefreshDecorations after
-	// layout builders complete (avoids AddChild flickering on every rebuild).
+	// removed from a section, the event fires before the decoration is actually
+	// removed from the container, so event-driven cleanup can't detect the change.
+	// This deferred check catches it during the subsequent rebuild.
 	for (TSharedPtr<FSectionOutlinerModel> SectionOutliner : OutlinerChildren.IterateSubList<FSectionOutlinerModel>())
 	{
 		if (!SectionOutliner)
@@ -109,16 +142,50 @@ TSharedPtr<FDecorationOutlinerModel> FDecorationModelStorageExtension::SyncDecor
 		}
 
 		UMovieSceneSection* Section = SectionOutliner->GetSection();
-		if (Section && !SequencerHelpers::HasOutlinerVisibleDecorations(Section))
+		if (!Section)
+		{
+			continue;
+		}
+
+		// Collect stale models from section outliner children
+		StaleModels.Reset();
+		FViewModelChildren SectionOutlinerChildren = SectionOutliner->GetChildList(EViewModelListType::Outliner);
+		for (TSharedPtr<FDecorationOutlinerModel> DecorationModel : SectionOutlinerChildren.IterateSubList<FDecorationOutlinerModel>())
 		{
-			if (TSharedPtr<FDecorationOutlinerModel> StaleModel = FindModel(Section))
+			if (!DecorationModel)
+			{
+				continue;
+			}
+
+			UObject* Decoration = DecorationModel->GetDecoration();
+			if (!Decoration)
+			{
+				StaleModels.Add(DecorationModel);
+				continue;
+			}
+
+			// Check if this decoration is still on the section
+			bool bStillExists = false;
+			for (const TObjectPtr<UObject>& SectionDecoration : Section->GetDecorations())
+			{
+				if (SectionDecoration.Get() == Decoration)
+				{
+					bStillExists = true;
+					break;
+				}
+			}
+
+			if (!bStillExists)
 			{
-				StaleModel->RemoveFromParent();
+				StaleModels.Add(DecorationModel);
 			}
 		}
-	}
 
-	return Result;
+		for (const TSharedPtr<FDecorationOutlinerModel>& StaleModel : StaleModels)
+		{
+			StaleModel->RemoveFromParent();
+		}
+	}
 }
 
 } // namespace UE::Sequencer
diff --git a/Engine/Source/Editor/Sequencer/Private/MVVM/ViewModels/DecorationOutlinerModel.cpp b/Engine/Source/Editor/Sequencer/Private/MVVM/ViewModels/DecorationOutlinerModel.cpp
--- a/Engine/Source/Editor/Sequencer/Private/MVVM/ViewModels/DecorationOutlinerModel.cpp
+++ b/Engine/Source/Editor/Sequencer/Private/MVVM/ViewModels/DecorationOutlinerModel.cpp
@@ -15,7 +15,6 @@
 #include "MVVM/ViewModels/ChannelModel.h"
 #include "MVVM/ViewModels/SequenceModel.h"
 #include "MVVM/ViewModels/SectionModel.h"
-#include "MVVM/ViewModels/SectionOutlinerModel.h"
 #include "MVVM/ViewModels/SequencerEditorViewModel.h"
 #include "MVVM/ViewModels/TrackModel.h"
 #include "MVVM/ViewModels/TrackModelLayoutBuilder.h"
@@ -27,13 +26,14 @@
 namespace UE::Sequencer
 {
 UE_SEQUENCER_DEFINE_CASTABLE(FDecorationOutlinerModel);
-UE_SEQUENCER_DEFINE_CASTABLE(FDecorationCategoryModel);
 
-FDecorationOutlinerModel::FDecorationOutlinerModel(UMovieSceneTrack* InParentTrack, UMovieSceneDecorationContainerObject* InDecorationContainer)
+FDecorationOutlinerModel::FDecorationOutlinerModel(UMovieSceneTrack* InParentTrack, UMovieSceneDecorationContainerObject* InDecorationContainer, UObject* InDecoration)
 	: WeakParentTrack(InParentTrack)
 	, WeakDecorationContainer(InDecorationContainer)
+	, WeakDecoration(InDecoration)
 {
 	RegisterChildList(&DecoratorList);
+	RegisterChildList(&SectionModelList);
 }
 
 FDecorationOutlinerModel::~FDecorationOutlinerModel()
@@ -42,43 +42,27 @@ FDecorationOutlinerModel::~FDecorationOutlinerModel()
 
 bool FDecorationOutlinerModel::NeedsReconstruct() const
 {
-	UMovieSceneDecorationContainerObject* Container = WeakDecorationContainer.Get();
-	if (!Container)
+	UObject* Decoration = WeakDecoration.Get();
+	if (!Decoration)
 	{
 		return false;
 	}
 
-	TArrayView<const TObjectPtr<UObject>> CurrentDecorations = Container->GetDecorations();
-	if (CurrentDecorations.Num() != LastConstructedDecorations.Num())
+	if (IMovieSceneChannelDecoration* ChannelDecoration = Cast<IMovieSceneChannelDecoration>(Decoration))
 	{
-		return true;
-	}
-
-	for (int32 Index = 0; Index < CurrentDecorations.Num(); ++Index)
-	{
-		if (CurrentDecorations[Index] != LastConstructedDecorations[Index].Get())
+		if (ChannelDecoration->ConsumeDecorationChannelProxyDirty())
 		{
 			return true;
 		}
 	}
 
-	// Check if any channel decoration flagged its proxy as dirty
-	for (const TObjectPtr<UObject>& Decoration : CurrentDecorations)
-	{
-		if (IMovieSceneChannelDecoration* ChannelDecoration = Cast<IMovieSceneChannelDecoration>(Decoration.Get()))
-		{
-			if (ChannelDecoration->ConsumeDecorationChannelProxyDirty())
-			{
-				return true;
-			}
-		}
-	}
 	return false;
 }
 
 void FDecorationOutlinerModel::OnConstruct()
 {
-	if (!WeakParentTrack.IsValid() || !WeakDecorationContainer.IsValid())
+	UObject* Decoration = WeakDecoration.Get();
+	if (!WeakParentTrack.IsValid() || !WeakDecorationContainer.IsValid() || !Decoration)
 	{
 		return;
 	}
@@ -91,7 +75,6 @@ void FDecorationOutlinerModel::OnConstruct()
 		return;
 	}
 
-	// Get ISequencer pointer from the ancestor view model.
 	TSharedPtr<FSequenceModel> SequenceModel = FindAncestorOfType<FSequenceModel>();
 	if (!SequenceModel)
 	{
@@ -109,7 +92,7 @@ void FDecorationOutlinerModel::OnConstruct()
 	{
 		return;
 	}
-	
+
 	TrackEditor = Sequencer->GetTrackEditor(WeakParentTrack.Get());
 	FViewModelChildren OutlinerChildren = GetTopLevelChannels();
 
@@ -118,9 +101,6 @@ void FDecorationOutlinerModel::OnConstruct()
 	FViewModelChildren RecycledChildren = RecycledModels.GetChildren();
 
 	// Remove old channel models from their parent section models before destroying proxies.
-	// Channel models are added to section models for ancestor lookup (see later in this function),
-	// and section models are reused across OnConstruct() calls. Without this cleanup, old channels
-	// with stale handles (pointing to destroyed proxies) would persist alongside new channels.
 	for (const TWeakViewModelPtr<FChannelModel>& WeakChannel : TrackedChannelModels)
 	{
 		if (TViewModelPtr<FChannelModel> OldChannel = WeakChannel.Pin())
@@ -130,211 +110,144 @@ void FDecorationOutlinerModel::OnConstruct()
 	}
 	TrackedChannelModels.Empty();
 
-	// Clear old channel proxies before rebuilding
-	ChannelDecorationProxies.Reset();
+	// Clear old channel proxy before rebuilding
+	ChannelDecorationProxy.Reset();
+
+	// Clear section model list
+	FViewModelChildren SectionChildren = GetChildrenForList(&SectionModelList);
+	SectionChildren.Empty();
 
-	TSharedPtr<FViewModel> OutlinerTail;
-	for (TObjectPtr<UObject> Decoration : WeakDecorationContainer->GetDecorations())
+	// Handle section-provider decorations: host sections directly in this model's track
+	// area without creating intermediate FSectionOutlinerModel wrappers.
+	if (IMovieSceneSectionProviderDecoration* SectionProvider = Cast<IMovieSceneSectionProviderDecoration>(Decoration))
 	{
-		// For decorations that have sections, create and add models for those sections.
-		if (IMovieSceneSectionProviderDecoration* SectionProvider = Cast<IMovieSceneSectionProviderDecoration>(Decoration))
+		ISequencerDecorationEditor* DecorationEditor = ISequencerDecorationEditor::FindDecorationEditor(Decoration->GetClass());
+		if (!DecorationEditor)
 		{
-			// Find the decoration editor for this decoration class
-			// Decorations are generic and can go on any track, so we can't use the parent track's
-			// track editor to create section interfaces. Instead, decoration editors are registered
-			// via modular features for each decoration type that provides sections.
-			ISequencerDecorationEditor* DecorationEditor = ISequencerDecorationEditor::FindDecorationEditor(Decoration->GetClass());
-			if (!DecorationEditor)
+			UE_LOG(LogSequencer, Warning, TEXT("No decoration editor registered for decoration class %s. Sections will not be displayed."), *Decoration->GetClass()->GetName());
+			return;
+		}
+
+		// Create layout builder once — calling RefreshLayout per section accumulates
+		// channels from all sections, matching how FTrackModel handles multi-section tracks.
+		FTrackModelLayoutBuilder LayoutBuilder(AsShared());
+
+		for (UMovieSceneSection* Section : SectionProvider->GetSections())
+		{
+			if (!Section)
 			{
-				UE_LOG(LogSequencer, Warning, TEXT("No decoration editor registered for decoration class %s. Sections will not be displayed."), *Decoration->GetClass()->GetName());
 				continue;
 			}
-
-			for (UMovieSceneSection* Section : SectionProvider->GetSections())
+			const TSharedPtr<ISequencerSection> SectionInterface = DecorationEditor->MakeSectionInterface(*Section, *WeakParentTrack.Get(), WeakParentTrack.Get()->FindObjectBindingGuid());
+			if (!SectionInterface)
 			{
-				if (!Section)
-				{
-					continue;
-				}
-				const TSharedPtr<ISequencerSection> SectionInterface = DecorationEditor->MakeSectionInterface(*Section, *WeakParentTrack.Get(), WeakParentTrack.Get()->FindObjectBindingGuid());
-				if (!SectionInterface)
-				{
-					continue;
-				}
-				const TSharedPtr<FSectionModel> SectionModel = SectionModelStorage->CreateModelForSection(Section, SectionInterface.ToSharedRef());
+				continue;
+			}
+			const TSharedPtr<FSectionModel> SectionModel = SectionModelStorage->CreateModelForSection(Section, SectionInterface.ToSharedRef());
 
-				TSharedPtr<FSectionOutlinerModel> SectionOutlinerModel;
-				// Search in RecycledChildren (where we moved existing models) rather than OutlinerChildren (which is now empty)
-				for (TSharedPtr<FSectionOutlinerModel> ExistingOutliner : RecycledChildren.IterateSubList<FSectionOutlinerModel>())
-				{
-					if (ExistingOutliner && ExistingOutliner->GetSection() == Section)
-					{
-						SectionOutlinerModel = ExistingOutliner;
-						break;
-					}
-				}
+			// Add section model to the track area list (renders the section in the track area)
+			SectionChildren.AddChild(SectionModel);
+			SectionModel->SetLinkedOutlinerItem(CastViewModel<IOutlinerExtension>(AsShared()));
 
-				if (!SectionOutlinerModel)
-				{
-					SectionOutlinerModel = MakeShared<FSectionOutlinerModel>(Section, SectionModel);
-				}
+			LayoutBuilder.RefreshLayout(SectionModel);
+		}
+	}
+	// Handle channel decorations
+	else if (IMovieSceneChannelDecoration* ChannelDecoration = Cast<IMovieSceneChannelDecoration>(Decoration))
+	{
+		if (!TrackEditor)
+		{
+			return;
+		}
 
-				OutlinerChildren.InsertChild(SectionOutlinerModel, OutlinerTail);
-				OutlinerTail = SectionOutlinerModel;
+		UMovieSceneSection* ParentSection = Cast<UMovieSceneSection>(WeakDecorationContainer.Get());
+		if (!ParentSection)
+		{
+			return;
+		}
 
-				SectionOutlinerModel->OnConstruct();
-				SectionOutlinerModel->GetTrackAreaChildren().AddChild(SectionModel);
-				SectionModel->SetLinkedOutlinerItem(SectionOutlinerModel);
+		TSharedRef<ISequencerSection> NewSectionInterface = TrackEditor->MakeSectionInterface(*ParentSection, *WeakParentTrack.Get(), WeakParentTrack.Get()->FindObjectBindingGuid());
+		const TSharedPtr<FSectionModel> SectionModel = SectionModelStorage->CreateModelForSection(ParentSection, NewSectionInterface);
 
-				FTrackModelLayoutBuilder LayoutBuilder(SectionOutlinerModel);
-				LayoutBuilder.RefreshLayout(SectionModel);
+		TSharedPtr<ISequencerSection> SectionInterface = SectionModel->GetSectionInterface();
 
-				// Refresh decorations after layout builder completes - the layout builder
-				// recycles the outliner child list, so decorations must be added after.
-				SectionOutlinerModel->RefreshDecorations();
-			}
-		}
-		// For decorations that add channels to existing sections
-		else if (IMovieSceneChannelDecoration* ChannelDecoration = Cast<IMovieSceneChannelDecoration>(Decoration))
-		{
-			if (!TrackEditor)
-			{
-				continue;
-			}
+		// Get channels from the decoration
+		FMovieSceneChannelProxyData ChannelProxyData;
+		ChannelDecoration->PopulateChannelProxy(ChannelProxyData);
 
-			// Check if the decoration container is a section (for section-level decorations)
-			UMovieSceneSection* ParentSection = Cast<UMovieSceneSection>(WeakDecorationContainer.Get());
-			if (!ParentSection)
-			{
-				continue;
-			}
+		// Build channel proxy and store it to keep channel handles valid
+		ChannelDecorationProxy = MakeShared<FMovieSceneChannelProxy>(MoveTemp(ChannelProxyData));
+		FMovieSceneChannelProxy& ChannelProxy = *ChannelDecorationProxy;
 
-			// Get or create the section model for the parent section
-			// We need to use the section interface from the section model (which is kept alive by the model)
-			// rather than creating a temporary one that would be destroyed when this function returns
-			TSharedRef<ISequencerSection> NewSectionInterface = TrackEditor->MakeSectionInterface(*ParentSection, *WeakParentTrack.Get(), WeakParentTrack.Get()->FindObjectBindingGuid());
-			const TSharedPtr<FSectionModel> SectionModel = SectionModelStorage->CreateModelForSection(ParentSection, NewSectionInterface);
-
-			// Use the section interface from the section model - this ensures it stays alive
-			// (CreateModelForSection may have returned an existing model with its own interface)
-			TSharedPtr<ISequencerSection> SectionInterface = SectionModel->GetSectionInterface();
-
-			// Get channels from the decoration
-			FMovieSceneChannelProxyData ChannelProxyData;
-			ChannelDecoration->PopulateChannelProxy(ChannelProxyData);
-
-			// Build channel proxy and store it to keep channel handles valid
-			TSharedPtr<FMovieSceneChannelProxy> ChannelProxyPtr = MakeShared<FMovieSceneChannelProxy>(MoveTemp(ChannelProxyData));
-			ChannelDecorationProxies.Add(ChannelProxyPtr);
-			FMovieSceneChannelProxy& ChannelProxy = *ChannelProxyPtr;
-
-			// Get the display name from the decoration class
-			FText DecorationDisplayName = Decoration->GetClass()->GetDisplayNameText();
-
-			// Create a category group for this decoration's channels
-			FName CategoryName(*Decoration->GetClass()->GetName());
-			TSharedPtr<FDecorationCategoryModel> CategoryOutlinerModel;
-			// Search in RecycledChildren (where we moved existing models) rather than OutlinerChildren (which is now empty)
-			for (TSharedPtr<FDecorationCategoryModel> ExistingCategory : RecycledChildren.IterateSubList<FDecorationCategoryModel>())
-			{
-				if (ExistingCategory && ExistingCategory->GetCategoryName() == CategoryName)
-				{
-					CategoryOutlinerModel = ExistingCategory;
-					break;
-				}
-			}
+		// Collect all channels with their metadata for sorting
+		struct FChannelSortData
+		{
+			FMovieSceneChannelHandle Handle;
+			const FMovieSceneChannelMetaData* MetaData;
+		};
+		TArray<FChannelSortData> SortedChannels;
+
+		for (const FMovieSceneChannelEntry& Entry : ChannelProxy.GetAllEntries())
+		{
+			TArrayView<const FMovieSceneChannelMetaData> AllMetaData = Entry.GetMetaData();
 
-			if (!CategoryOutlinerModel)
+			for (int32 ChannelIndex = 0; ChannelIndex < AllMetaData.Num(); ++ChannelIndex)
 			{
-				CategoryOutlinerModel = MakeShared<FDecorationCategoryModel>(CategoryName, DecorationDisplayName, Decoration);
+				FChannelSortData SortData;
+				SortData.Handle = ChannelProxy.MakeHandle(Entry.GetChannelTypeName(), ChannelIndex);
+				SortData.MetaData = &AllMetaData[ChannelIndex];
+				SortedChannels.Add(SortData);
 			}
+		}
 
-			OutlinerChildren.InsertChild(CategoryOutlinerModel, OutlinerTail);
-			OutlinerTail = CategoryOutlinerModel;
+		// Sort by SortOrder
+		SortedChannels.Sort([](const FChannelSortData& A, const FChannelSortData& B)
+		{
+			return A.MetaData->SortOrder < B.MetaData->SortOrder;
+		});
 
-			// Collect all channels with their metadata for sorting
-			struct FChannelSortData
-			{
-				FMovieSceneChannelHandle Handle;
-				const FMovieSceneChannelMetaData* MetaData;
-			};
-			TArray<FChannelSortData> SortedChannels;
+		// Add channels directly under this decoration model
+		TSharedPtr<FViewModel> ChannelTail;
+		for (const FChannelSortData& ChannelData : SortedChannels)
+		{
+			const FMovieSceneChannelMetaData& MetaData = *ChannelData.MetaData;
+			FMovieSceneChannelHandle ChannelHandle = ChannelData.Handle;
 
-			for (const FMovieSceneChannelEntry& Entry : ChannelProxy.GetAllEntries())
+			FName ChannelName = MetaData.Name;
+			TSharedPtr<FChannelGroupOutlinerModel> ChannelGroupModel;
+			for (TSharedPtr<FChannelGroupOutlinerModel> ExistingGroup : OutlinerChildren.IterateSubList<FChannelGroupOutlinerModel>())
 			{
-				TArrayView<const FMovieSceneChannelMetaData> AllMetaData = Entry.GetMetaData();
-
-				for (int32 ChannelIndex = 0; ChannelIndex < AllMetaData.Num(); ++ChannelIndex)
+				if (ExistingGroup && ExistingGroup->GetChannelName() == ChannelName)
 				{
-					FChannelSortData SortData;
-					SortData.Handle = ChannelProxy.MakeHandle(Entry.GetChannelTypeName(), ChannelIndex);
-					SortData.MetaData = &AllMetaData[ChannelIndex];
-					SortedChannels.Add(SortData);
+					ChannelGroupModel = ExistingGroup;
+					break;
 				}
 			}
 
-			// Sort by SortOrder
-			SortedChannels.Sort([](const FChannelSortData& A, const FChannelSortData& B)
+			if (!ChannelGroupModel)
 			{
-				return A.MetaData->SortOrder < B.MetaData->SortOrder;
-			});
-
-			// Add channels from the decoration to the category in sorted order
-			FViewModelChildren CategoryChildren = CategoryOutlinerModel->GetChildList(EViewModelListType::Outliner);
-			TSharedPtr<FViewModel> ChannelTail;
-			for (const FChannelSortData& ChannelData : SortedChannels)
+				ChannelGroupModel = MakeShared<FChannelGroupOutlinerModel>(ChannelName, MetaData.DisplayText, MetaData.GetTooltipTextDelegate);
+			}
+			else
 			{
-				const FMovieSceneChannelMetaData& MetaData = *ChannelData.MetaData;
-				FMovieSceneChannelHandle ChannelHandle = ChannelData.Handle;
-
-				// Create channel group outliner model
-				FName ChannelName = MetaData.Name;
-				TSharedPtr<FChannelGroupOutlinerModel> ChannelGroupModel;
-				for (TSharedPtr<FChannelGroupOutlinerModel> ExistingGroup : CategoryChildren.IterateSubList<FChannelGroupOutlinerModel>())
-				{
-					if (ExistingGroup && ExistingGroup->GetChannelName() == ChannelName)
-					{
-						ChannelGroupModel = ExistingGroup;
-						break;
-					}
-				}
-
-				if (!ChannelGroupModel)
-				{
-					ChannelGroupModel = MakeShared<FChannelGroupOutlinerModel>(ChannelName, MetaData.DisplayText, MetaData.GetTooltipTextDelegate);
-				}
-				else
-				{
-					// Clear old channels when reusing an existing group to avoid stale handles
-					// from a previous proxy that has been destroyed
-					ChannelGroupModel->OnRecycle();
-				}
+				// Clear old channels when reusing an existing group to avoid stale handles
+				ChannelGroupModel->OnRecycle();
+			}
 
-				CategoryChildren.InsertChild(ChannelGroupModel, ChannelTail);
-				ChannelTail = ChannelGroupModel;
+			OutlinerChildren.InsertChild(ChannelGroupModel, ChannelTail);
+			ChannelTail = ChannelGroupModel;
 
-				// Create channel model
-				TSharedPtr<FChannelModel> ChannelModel = MakeShared<FChannelModel>(ChannelName, SectionInterface, ChannelHandle);
-				ChannelGroupModel->AddChannel(ChannelModel);
-				ChannelModel->SetLinkedOutlinerItem(ChannelGroupModel.ToSharedRef());
-				ChannelModel->Initialize(SectionInterface, ChannelHandle);
+			TSharedPtr<FChannelModel> ChannelModel = MakeShared<FChannelModel>(ChannelName, SectionInterface, ChannelHandle);
+			ChannelGroupModel->AddChannel(ChannelModel);
+			ChannelModel->SetLinkedOutlinerItem(ChannelGroupModel.ToSharedRef());
+			ChannelModel->Initialize(SectionInterface, ChannelHandle);
 
-				// Add channel model as child of section model so FindAncestorOfType<FSectionModel>() works
-				// This is required for ArrangeVirtualTrackLaneView() to get the section range for key drawing
-				SectionModel->GetChildList(FTrackModel::GetTopLevelChannelType()).AddChild(ChannelModel);
+			// Add channel model as child of section model so FindAncestorOfType<FSectionModel>() works
+			SectionModel->GetChildList(FTrackModel::GetTopLevelChannelType()).AddChild(ChannelModel);
 
-				// Track this channel so we can remove it from section model during cleanup
-				TrackedChannelModels.Add(ChannelModel);
-			}
+			TrackedChannelModels.Add(ChannelModel);
 		}
 	}
-
-	// Store the decoration state so NeedsReconstruct() can detect changes
-	LastConstructedDecorations.Reset();
-	for (const TObjectPtr<UObject>& Decoration : WeakDecorationContainer->GetDecorations())
-	{
-		LastConstructedDecorations.Add(Decoration);
-	}
 }
 
 /*------------------------------------------------------------------------------------
@@ -347,24 +260,20 @@ void FDecorationOutlinerModel::BuildContextMenu(FMenuBuilder& MenuBuilder)
 	{
 		return;
 	}
-	
+
 	const TSharedPtr<FSequencer> Sequencer = EditorViewModel->GetSequencerImpl();
 	if (!Sequencer.IsValid())
 	{
 		return;
 	}
-	
-	if (WeakDecorationContainer.IsValid())
+
+	UObject* Decoration = WeakDecoration.Get();
+	if (Decoration && Decoration->Implements<UMovieSceneSectionProviderDecoration>())
 	{
-		for (TObjectPtr<UObject> Decoration : WeakDecorationContainer->GetDecorations())
-		{
-			if (Decoration && Decoration->Implements<UMovieSceneSectionProviderDecoration>())
-			{
-				const TArray<TWeakObjectPtr<>> TrackAreaModels = SequencerHelpers::GetSectionObjectsFromTrackAreaModels(GetTrackAreaModelList());
-				SequencerHelpers::BuildEditSectionMenu(Sequencer, TrackAreaModels, MenuBuilder, true);
-			}
-		}
+		const TArray<TWeakObjectPtr<>> TrackAreaModels = SequencerHelpers::GetSectionObjectsFromTrackAreaModels(GetTrackAreaModelList());
+		SequencerHelpers::BuildEditSectionMenu(Sequencer, TrackAreaModels, MenuBuilder, true);
 	}
+
 	FOutlinerItemModel::BuildContextMenu(MenuBuilder);
 }
 
@@ -380,7 +289,7 @@ FOutlinerSizing FDecorationOutlinerModel::GetOutlinerSizing() const
 {
 	const FViewDensityInfo Density = GetEditor()->GetViewDensity();
 	float Height = Density.UniformHeight.Get(SequencerLayoutConstants::SectionAreaDefaultHeight);
-	for (const TSharedPtr<FSectionModel> Section : DecoratorList.Iterate<FSectionModel>())
+	for (const TSharedPtr<FSectionModel> Section : SectionModelList.Iterate<FSectionModel>())
 	{
 		Height = Section->GetSectionInterface()->GetSectionHeight(Density);
 		break;
@@ -401,7 +310,16 @@ TSharedPtr<SWidget> FDecorationOutlinerModel::CreateOutlinerViewForColumn(const
 
 FText FDecorationOutlinerModel::GetLabel() const
 {
-	return LOCTEXT("ModifiersLabel", "Modifiers");
+	if (UObject* Decoration = WeakDecoration.Get())
+	{
+		return Decoration->GetClass()->GetDisplayNameText();
+	}
+	return FText::GetEmpty();
+}
+
+FViewModelChildren FDecorationOutlinerModel::GetSectionModels()
+{
+	return GetChildrenForList(&SectionModelList);
 }
 
 FViewModelChildren FDecorationOutlinerModel::GetTopLevelChannels()
@@ -422,7 +340,7 @@ FTrackAreaParameters FDecorationOutlinerModel::GetTrackAreaParameters() const
 
 FViewModelVariantIterator FDecorationOutlinerModel::GetTrackAreaModelList() const
 {
-	return &DecoratorList;
+	return &SectionModelList;
 }
 
 FViewModelVariantIterator FDecorationOutlinerModel::GetTopLevelChildTrackAreaModels() const
@@ -436,16 +354,16 @@ FViewModelVariantIterator FDecorationOutlinerModel::GetTopLevelChildTrackAreaMod
 
 bool FDecorationOutlinerModel::CanDelete(FText* OutErrorMessage) const
 {
-	// This outliner only exists if there are decorations to display, so assume that this outliner can always be deleted if it exists.
-	return true;
+	return WeakDecoration.IsValid() && WeakDecorationContainer.IsValid();
 }
 
 void FDecorationOutlinerModel::Delete()
 {
 	UMovieSceneDecorationContainerObject* DecorationContainer = WeakDecorationContainer.Get();
 	UMovieSceneTrack* ParentTrack = WeakParentTrack.Get();
+	UObject* Decoration = WeakDecoration.Get();
 
-	if (!DecorationContainer || !ParentTrack)
+	if (!DecorationContainer || !ParentTrack || !Decoration)
 	{
 		return;
 	}
@@ -453,20 +371,7 @@ void FDecorationOutlinerModel::Delete()
 	ParentTrack->Modify();
 	DecorationContainer->Modify();
 
-	TArray<UClass*> DecorationsToRemove;
-
-	for (const UObject* Decoration : DecorationContainer->GetDecorations())
-	{
-		if (Decoration)
-		{
-			DecorationsToRemove.Add(Decoration->GetClass());
-		}
-	}
-
-	for (UClass* DecorationClass : DecorationsToRemove)
-	{
-		DecorationContainer->RemoveDecoration(DecorationClass);
-	}
+	DecorationContainer->RemoveDecoration(Decoration->GetClass());
 }
 
 /*-----------------------------------------------------------------------------
@@ -477,8 +382,8 @@ ELockableLockState FDecorationOutlinerModel::GetLockState() const
 {
 	int32 NumSections = 0;
 	int32 NumLockedSections = 0;
-	
-	for (const TViewModelPtr<FSectionModel>& Section : DecoratorList.Iterate<FSectionModel>())
+
+	for (const TViewModelPtr<FSectionModel>& Section : SectionModelList.Iterate<FSectionModel>())
 	{
 		++NumSections;
 
@@ -499,7 +404,7 @@ ELockableLockState FDecorationOutlinerModel::GetLockState() const
 
 void FDecorationOutlinerModel::SetIsLocked(bool bIsLocked)
 {
-	for (const TViewModelPtr<FSectionModel>& Section : DecoratorList.Iterate<FSectionModel>())
+	for (const TViewModelPtr<FSectionModel>& Section : SectionModelList.Iterate<FSectionModel>())
 	{
 		UMovieSceneSection* SectionObject = Section->GetSection();
 		if (SectionObject && SectionObject->IsLocked() != bIsLocked)
@@ -511,76 +416,11 @@ void FDecorationOutlinerModel::SetIsLocked(bool bIsLocked)
 }
 
 
-TArray<UObject*> FDecorationOutlinerModel::GetDecorations() const
+UObject* FDecorationOutlinerModel::GetDecoration() const
 {
-	TArray<UObject*> Result;
-
-	if (UMovieSceneDecorationContainerObject* Container = WeakDecorationContainer.Get())
-	{
-		for (const TObjectPtr<UObject>& Decoration : Container->GetDecorations())
-		{
-			if (Decoration)
-			{
-				Result.Add(Decoration.Get());
-			}
-		}
-	}
-
-	return Result;
-}
-
-void FDecorationOutlinerModel::DeleteDecoration(UClass* DecorationClass)
-{
-	UMovieSceneDecorationContainerObject* DecorationContainer = WeakDecorationContainer.Get();
-	UMovieSceneTrack* ParentTrack = WeakParentTrack.Get();
-
-	if (!DecorationContainer || !ParentTrack || !DecorationClass)
-	{
-		return;
-	}
-
-	ParentTrack->Modify();
-	DecorationContainer->Modify();
-
-	DecorationContainer->RemoveDecoration(DecorationClass);
-}
-
-/*------------------------------------------------------------------------------------
-	FDecorationCategoryModel
- -----------------------------------------------------------------------------------*/
-
-FDecorationCategoryModel::FDecorationCategoryModel(FName InCategoryName, const FText& InDisplayText, UObject* InDecoration)
-	: FCategoryGroupModel(InCategoryName, InDisplayText, FGetMovieSceneTooltipText())
-	, WeakDecoration(InDecoration)
-{
-}
-
-bool FDecorationCategoryModel::CanDelete(FText* OutErrorMessage) const
-{
-	// Can delete if we have a valid decoration and parent decoration model
-	if (!WeakDecoration.IsValid())
-	{
-		return false;
-	}
-
-	TSharedPtr<FDecorationOutlinerModel> DecorationModel = FindAncestorOfType<FDecorationOutlinerModel>();
-	return DecorationModel.IsValid();
-}
-
-void FDecorationCategoryModel::Delete()
-{
-	if (!WeakDecoration.IsValid())
-	{
-		return;
-	}
-
-	TSharedPtr<FDecorationOutlinerModel> DecorationModel = FindAncestorOfType<FDecorationOutlinerModel>();
-	if (DecorationModel)
-	{
-		DecorationModel->DeleteDecoration(WeakDecoration->GetClass());
-	}
+	return WeakDecoration.Get();
 }
 
 }
 
-#undef LOCTEXT_NAMESPACE
\ No newline at end of file
+#undef LOCTEXT_NAMESPACE
diff --git a/Engine/Source/Editor/Sequencer/Private/MVVM/ViewModels/SectionOutlinerModel.cpp b/Engine/Source/Editor/Sequencer/Private/MVVM/ViewModels/SectionOutlinerModel.cpp
--- a/Engine/Source/Editor/Sequencer/Private/MVVM/ViewModels/SectionOutlinerModel.cpp
+++ b/Engine/Source/Editor/Sequencer/Private/MVVM/ViewModels/SectionOutlinerModel.cpp
@@ -89,10 +89,6 @@ void FSectionOutlinerModel::OnConstruct()
 		Section->EventHandlers.Link(SectionEventHandler, this);
 	}
 
-	// Mark decorations for refresh. Don't call RefreshDecorations() here because
-	// callers typically run layout builders after OnConstruct, which recycle the
-	// outliner child list. RefreshDecorations must be called after layout completes.
-	bDecorationsNeedRefresh = true;
 }
 
 FViewModelChildren FSectionOutlinerModel::GetTrackAreaChildren()
@@ -718,24 +714,16 @@ void FSectionOutlinerModel::SetIsSoloed(bool bIsSoloed)
 
 void FSectionOutlinerModel::OnDecorationAdded(UObject*)
 {
-	bDecorationsNeedRefresh = true;
 	RefreshDecorations();
 }
 
 void FSectionOutlinerModel::OnDecorationRemoved(UObject*)
 {
-	bDecorationsNeedRefresh = true;
 	RefreshDecorations();
 }
 
 void FSectionOutlinerModel::RefreshDecorations()
 {
-	if (!bDecorationsNeedRefresh)
-	{
-		return;
-	}
-	bDecorationsNeedRefresh = false;
-
 	UMovieSceneSection* Section = WeakSection.Get();
 	if (!Section)
 	{
@@ -761,7 +749,7 @@ void FSectionOutlinerModel::RefreshDecorations()
 	}
 
 	FViewModelChildren OutlinerChildren = GetChildList(EViewModelListType::Outliner);
-	DecorationModelStorage->SyncDecorationModel(Section, Track, OutlinerChildren);
+	DecorationModelStorage->SyncDecorationModels(Section, Track, OutlinerChildren);
 }
 
 } // namespace Sequencer
diff --git a/Engine/Source/Editor/Sequencer/Private/MVVM/ViewModels/TrackModel.cpp b/Engine/Source/Editor/Sequencer/Private/MVVM/ViewModels/TrackModel.cpp
--- a/Engine/Source/Editor/Sequencer/Private/MVVM/ViewModels/TrackModel.cpp
+++ b/Engine/Source/Editor/Sequencer/Private/MVVM/ViewModels/TrackModel.cpp
@@ -313,7 +313,7 @@ void FTrackModel::ForceUpdate()
 	// Handle track-level decorations
 	if (DecorationModelStorage)
 	{
-		DecorationModelStorage->SyncDecorationModel(Track, Track, OutlinerChildren);
+		DecorationModelStorage->SyncDecorationModels(Track, Track, OutlinerChildren);
 	}
 
 	PreviousLayoutNumRows = NumRows;
diff --git a/Engine/Source/Editor/Sequencer/Private/Widgets/Sidebar/SequencerSelectionDrawer.cpp b/Engine/Source/Editor/Sequencer/Private/Widgets/Sidebar/SequencerSelectionDrawer.cpp
--- a/Engine/Source/Editor/Sequencer/Private/Widgets/Sidebar/SequencerSelectionDrawer.cpp
+++ b/Engine/Source/Editor/Sequencer/Private/Widgets/Sidebar/SequencerSelectionDrawer.cpp
@@ -424,21 +424,10 @@ void FSequencerSelectionDrawer::BuildOutlinerDetails(const TSharedRef<FSequencer
         {
         	ChannelGroups.Add(ChannelGroupOutlinerModel);
         }
-		// Handle decoration container (the "Modifiers" row) - check before FSectionOutlinerModel
+		// Handle per-decoration outliner row - check before FSectionOutlinerModel
 		else if (const TViewModelPtr<FDecorationOutlinerModel> DecorationOutlinerModel = OutlinerItem.ImplicitCast())
 		{
-			// Get decorations directly from the model
-			// Section-providing decorations will have their sections displayed under the decoration header
-			// by BuildEditDecorationMenu, so we don't collect sections separately here
-			for (UObject* Decoration : DecorationOutlinerModel->GetDecorations())
-			{
-				WeakDecorationObjects.AddUnique(Decoration);
-			}
-		}
-		// Handle individual decoration category (e.g., "Root Motion Settings")
-		else if (const TViewModelPtr<FDecorationCategoryModel> DecorationCategoryModel = OutlinerItem.ImplicitCast())
-		{
-			if (UObject* Decoration = DecorationCategoryModel->GetDecoration())
+			if (UObject* Decoration = DecorationOutlinerModel->GetDecoration())
 			{
 				WeakDecorationObjects.AddUnique(Decoration);
 			}
diff --git a/Engine/Source/Editor/Sequencer/Public/MVVM/DecorationModelStorageExtension.h b/Engine/Source/Editor/Sequencer/Public/MVVM/DecorationModelStorageExtension.h
--- a/Engine/Source/Editor/Sequencer/Public/MVVM/DecorationModelStorageExtension.h
+++ b/Engine/Source/Editor/Sequencer/Public/MVVM/DecorationModelStorageExtension.h
@@ -21,7 +21,7 @@ namespace Sequencer
 class FDecorationOutlinerModel;
 struct FViewModelChildren;
 
-// Central storage for decoration outliner models, keyed by decoration container.
+// Central storage for decoration outliner models, keyed by individual decoration objects.
 // Registered on FSequenceModel's shared data, following FSectionModelStorageExtension pattern.
 class FDecorationModelStorageExtension
 	: public IDynamicExtension
@@ -32,24 +32,22 @@ public:
 
 	UE_API FDecorationModelStorageExtension();
 
-	// Finds an existing model for the container, or creates a new one.
-	// ParentTrack is passed through to FDecorationOutlinerModel's constructor.
-	// bOutNewlyCreated is set to true if a new model was created, false if an existing one was found.
+	// Finds an existing model for the decoration, or creates a new one.
 	UE_API TSharedPtr<FDecorationOutlinerModel> FindOrCreateModel(
+		UObject* Decoration,
 		UMovieSceneDecorationContainerObject* Container,
 		UMovieSceneTrack* ParentTrack,
 		bool& bOutNewlyCreated);
 
-	// Finds an existing model for the container, returns nullptr if not found.
+	// Finds an existing model for the decoration, returns nullptr if not found.
 	UE_API TSharedPtr<FDecorationOutlinerModel> FindModel(
-		const UMovieSceneDecorationContainerObject* Container) const;
-
-	// Ensures the decoration model for the given container is present in OutlinerChildren
-	// if the container has visible decorations, or removed if it doesn't. Handles all
-	// lifecycle: creation, re-adoption from recycled, NeedsReconstruct checks, and stale
-	// model removal. Callers don't need to check HasOutlinerVisibleDecorations - this
-	// method handles that internally. Returns the model if one was added, or nullptr.
-	UE_API TSharedPtr<FDecorationOutlinerModel> SyncDecorationModel(
+		const UObject* Decoration) const;
+
+	// Ensures decoration models for the given container are present in OutlinerChildren.
+	// Creates one FDecorationOutlinerModel per visible decoration on the container.
+	// Handles creation, re-adoption from recycled, NeedsReconstruct checks, and stale
+	// model removal. Also cleans up stale decoration models on child section outliners.
+	UE_API void SyncDecorationModels(
 		UMovieSceneDecorationContainerObject* Container,
 		UMovieSceneTrack* ParentTrack,
 		FViewModelChildren& OutlinerChildren);
@@ -58,7 +56,7 @@ public:
 
 private:
 
-	TMap<FObjectKey, TWeakPtr<FDecorationOutlinerModel>> ContainerToModel;
+	TMap<FObjectKey, TWeakPtr<FDecorationOutlinerModel>> DecorationToModel;
 };
 
 } // namespace Sequencer
diff --git a/Engine/Source/Editor/Sequencer/Public/MVVM/ViewModels/DecorationOutlinerModel.h b/Engine/Source/Editor/Sequencer/Public/MVVM/ViewModels/DecorationOutlinerModel.h
--- a/Engine/Source/Editor/Sequencer/Public/MVVM/ViewModels/DecorationOutlinerModel.h
+++ b/Engine/Source/Editor/Sequencer/Public/MVVM/ViewModels/DecorationOutlinerModel.h
@@ -5,11 +5,10 @@
 #include "OutlinerItemModel.h"
 #include "MVVM/Extensions/IDeletableExtension.h"
 #include "MVVM/Extensions/ILockableExtension.h"
+#include "MVVM/Extensions/ISectionOwnerExtension.h"
 #include "MVVM/Extensions/ITopLevelChannelHolderExtension.h"
 #include "MVVM/Extensions/ITrackAreaExtension.h"
-#include "MVVM/ViewModels/CategoryModel.h"
 
-class UMovieSceneTrackRowDecoration;
 struct FMovieSceneChannelProxy;
 
 #define UE_API SEQUENCER_API
@@ -17,8 +16,15 @@ struct FMovieSceneChannelProxy;
 namespace UE::Sequencer
 {
 
+// Outliner model for a single decoration on a track or section.
+// One instance is created per visible decoration (IMovieSceneChannelDecoration or
+// IMovieSceneSectionProviderDecoration). For section-provider decorations, section
+// models are hosted directly in the track area and channel items appear as outliner
+// children. For channel decorations, FChannelGroupOutlinerModel instances are direct
+// outliner children.
 class FDecorationOutlinerModel
 	: public FOutlinerItemModel
+	, public ISectionOwnerExtension
 	, public ITrackAreaExtension
 	, public ILockableExtension
 	, public IDeletableExtension
@@ -27,12 +33,13 @@ class FDecorationOutlinerModel
 public:
 	UE_SEQUENCER_DECLARE_CASTABLE_API(SEQUENCER_API, FDecorationOutlinerModel
 		, FOutlinerItemModel
+		, ISectionOwnerExtension
 		, ITrackAreaExtension
 		, ILockableExtension
 		, IDeletableExtension
 		, ITopLevelChannelHolderExtension);
 
-	UE_API explicit FDecorationOutlinerModel(UMovieSceneTrack* InParentTrack, UMovieSceneDecorationContainerObject* InDecorationContainer);
+	UE_API explicit FDecorationOutlinerModel(UMovieSceneTrack* InParentTrack, UMovieSceneDecorationContainerObject* InDecorationContainer, UObject* InDecoration);
 	UE_API virtual ~FDecorationOutlinerModel() override;
 
 	/* FViewModel Interface */
@@ -47,6 +54,9 @@ public:
 	UE_API virtual TSharedPtr<SWidget> CreateOutlinerViewForColumn(const FCreateOutlinerViewParams& InParams, const FName& InColumnName) override;
 	UE_API virtual FText GetLabel() const override;
 
+	/* ISectionOwnerExtension */
+	UE_API virtual FViewModelChildren GetSectionModels() override;
+
 	/* ITopLevelChannelHolderExtension */
 	UE_API virtual FViewModelChildren GetTopLevelChannels() override;
 
@@ -65,58 +75,35 @@ public:
 
 	UE_API TWeakObjectPtr<UMovieSceneDecorationContainerObject> GetDecorationContainer() const { return WeakDecorationContainer; }
 
-	// Get all decoration objects from this model's container. Used for sidebar property display.
-	UE_API TArray<UObject*> GetDecorations() const;
-
-	// Delete a specific decoration by class. Called by FDecorationCategoryModel when its row is deleted.
-	UE_API void DeleteDecoration(UClass* DecorationClass);
+	// Get the decoration object this model represents. Used for sidebar property display.
+	UE_API UObject* GetDecoration() const;
 
-	// Returns true if the container's decorations have changed since the last OnConstruct.
+	// Returns true if the decoration's state has changed since the last OnConstruct.
 	UE_API bool NeedsReconstruct() const;
 
 private:
 
+	// Outliner children: channel groups for both section-provider and channel decorations.
 	FViewModelListHead DecoratorList{EViewModelListType::Outliner};
 
+	// Section models for section-provider decorations. Used for the track area display.
+	FViewModelListHead SectionModelList{EViewModelListType::TrackArea};
+
 	TWeakObjectPtr<UMovieSceneTrack> WeakParentTrack;
 
 	TWeakObjectPtr<UMovieSceneDecorationContainerObject> WeakDecorationContainer;
 
+	// The specific decoration this model represents.
+	TWeakObjectPtr<UObject> WeakDecoration;
+
 	TSharedPtr<ISequencerTrackEditor> TrackEditor;
 
-	// Channel proxies for IMovieSceneChannelDecoration decorations. Stored to keep channel handles valid.
-	TArray<TSharedPtr<FMovieSceneChannelProxy>> ChannelDecorationProxies;
+	// Channel proxy for IMovieSceneChannelDecoration. Stored to keep channel handles valid.
+	TSharedPtr<FMovieSceneChannelProxy> ChannelDecorationProxy;
 
 	// Tracks channel models added to section models so they can be removed during cleanup.
 	TArray<TWeakViewModelPtr<FChannelModel>> TrackedChannelModels;
 
-	// Decoration objects present when OnConstruct last ran. Used to detect when rebuild is needed.
-	TArray<TWeakObjectPtr<UObject>> LastConstructedDecorations;
-
-};
-
-/**
- * Specialized category model for decoration channels that supports deletion of the specific decoration.
- * Used instead of FCategoryGroupModel when displaying channels from IMovieSceneChannelDecoration.
- */
-class FDecorationCategoryModel
-	: public FCategoryGroupModel
-{
-public:
-	UE_SEQUENCER_DECLARE_CASTABLE_API(UE_API, FDecorationCategoryModel, FCategoryGroupModel);
-
-	UE_API explicit FDecorationCategoryModel(FName InCategoryName, const FText& InDisplayText, UObject* InDecoration);
-
-	/* IDeletableExtension - override to delete just this decoration */
-	UE_API virtual bool CanDelete(FText* OutErrorMessage) const override;
-	UE_API virtual void Delete() override;
-
-	/** Get the decoration object this category represents */
-	UE_API UObject* GetDecoration() const { return WeakDecoration.Get(); }
-
-private:
-	/** The decoration object this category represents */
-	TWeakObjectPtr<UObject> WeakDecoration;
 };
 
 #undef UE_API
