Sequencer: Move private Curve Editor files to CurveEditor folder.

Specifically, moved files out from
- MVVM/Views
- Widget/CurveEditor

This makes all curve editor related implementations easier to find.

Also removed some unused includes and sorted them alphabetically.

#Synced-CL 51334859

diff --git a/Engine/Source/Editor/Sequencer/Private/CurveEditor/Integration/Selection/CurveModelSyncer.cpp b/Engine/Source/Editor/Sequencer/Private/CurveEditor/Integration/Selection/CurveModelSyncer.cpp
--- a/Engine/Source/Editor/Sequencer/Private/CurveEditor/Integration/Selection/CurveModelSyncer.cpp
+++ b/Engine/Source/Editor/Sequencer/Private/CurveEditor/Integration/Selection/CurveModelSyncer.cpp
@@ -3,15 +3,15 @@
 #include "CurveModelSyncer.h"
 
 #include "CurveEditor.h"
-#include "Sequencer.h"
+#include "CurveEditor/Views/CurveEditorWidgetOwner.h"
 #include "Filters/ViewModels/Linking/ILinkedFilterViewModel.h"
-#include "Modification/Utils/ScopedSelectionChange.h"
-#include "MVVM/SharedViewModelData.h"
-#include "MVVM/ViewModelPtr.h"
 #include "MVVM/Extensions/IOutlinerExtension.h"
 #include "MVVM/Selection/Selection.h"
+#include "MVVM/SharedViewModelData.h"
+#include "MVVM/ViewModelPtr.h"
 #include "MVVM/ViewModels/ViewModelIterators.h"
-#include "MVVM/Views/CurveEditor/CurveEditorWidgetOwner.h"
+#include "Modification/Utils/ScopedSelectionChange.h"
+#include "Sequencer.h"
 #include "Tree/CurveEditorTree.h"
 
 namespace UE::Sequencer
diff --git a/Engine/Source/Editor/Sequencer/Private/CurveEditor/Integration/Selection/SequencerSelectionSyncer.cpp b/Engine/Source/Editor/Sequencer/Private/CurveEditor/Integration/Selection/SequencerSelectionSyncer.cpp
--- a/Engine/Source/Editor/Sequencer/Private/CurveEditor/Integration/Selection/SequencerSelectionSyncer.cpp
+++ b/Engine/Source/Editor/Sequencer/Private/CurveEditor/Integration/Selection/SequencerSelectionSyncer.cpp
@@ -3,11 +3,11 @@
 #include "SequencerSelectionSyncer.h"
 
 #include "CurveEditor.h"
-#include "Sequencer.h"
+#include "CurveEditor/Views/SelectionSyncUtils.h"
 #include "MVVM/Selection/Selection.h"
-#include "MVVM/Views/CurveEditor/SelectionSyncUtils.h"
-#include "Misc/ScopeExit.h"
 #include "MVVM/Views/SOutlinerView.h"
+#include "Misc/ScopeExit.h"
+#include "Sequencer.h"
 
 namespace UE::Sequencer
 {
diff --git a/Engine/Source/Editor/Sequencer/Private/CurveEditor/SequencerCurveEditorApp.cpp b/Engine/Source/Editor/Sequencer/Private/CurveEditor/SequencerCurveEditorApp.cpp
--- a/Engine/Source/Editor/Sequencer/Private/CurveEditor/SequencerCurveEditorApp.cpp
+++ b/Engine/Source/Editor/Sequencer/Private/CurveEditor/SequencerCurveEditorApp.cpp
@@ -7,14 +7,14 @@
 #include "Filters/ViewModels/Linking/LinkedFilterViewModel.h"
 #include "Filters/ViewModels/Linking/MakeLinkedFilterViewModelArgs.h"
 #include "ICurveEditorModule.h"
+#include "MVVM/CurveEditorExtension.h"
 #include "MVVM/ViewModels/SequencerEditorViewModelImpl.h"
 #include "Specializations/SequencerCurveEditor.h"
 #include "Specializations/SequencerCurveEditorBounds.h"
 #include "Specializations/SequencerCurveEditorToolbarExtender.h"
 #include "Specializations/SequencerCurveModelUtils.h"
 #include "TimeSliderArgs.h"
-#include "MVVM/CurveEditorExtension.h"
-#include "Widgets/CurveEditor/SequencerCurveEditorTimeSliderController.h"
+#include "Widgets/SequencerCurveEditorTimeSliderController.h"
 
 namespace UE::Sequencer
 {
diff --git a/Engine/Source/Editor/Sequencer/Private/CurveEditor/SequencerCurveEditorApp.h b/Engine/Source/Editor/Sequencer/Private/CurveEditor/SequencerCurveEditorApp.h
--- a/Engine/Source/Editor/Sequencer/Private/CurveEditor/SequencerCurveEditorApp.h
+++ b/Engine/Source/Editor/Sequencer/Private/CurveEditor/SequencerCurveEditorApp.h
@@ -2,9 +2,9 @@
 
 #pragma once
 
+#include "CurveEditor/Views/CurveEditorWidgetOwner.h"
 #include "Integration/CurveEditorCommandBinder.h"
 #include "Integration/Selection/SelectionSyncRootLogic.h"
-#include "MVVM/Views/CurveEditor/CurveEditorWidgetOwner.h"
 #include "Templates/SharedPointer.h"
 #include "Templates/UnrealTemplate.h"
 
diff --git a/Engine/Source/Editor/Sequencer/Private/CurveEditor/Views/CurveEditorWidgetOwner.cpp b/Engine/Source/Editor/Sequencer/Private/CurveEditor/Views/CurveEditorWidgetOwner.cpp
new file mode 100644
--- /dev/null
+++ b/Engine/Source/Editor/Sequencer/Private/CurveEditor/Views/CurveEditorWidgetOwner.cpp
@@ -0,0 +1,258 @@
+﻿// Copyright Epic Games, Inc. All Rights Reserved.
+
+#include "CurveEditorWidgetOwner.h"
+
+#include "CurveEditor/Widgets/SSequencerCurveEditor.h"
+#include "CurveEditor/Widgets/Tree/SCurveTreeContent.h"
+#include "CurveEditor/Widgets/Tree/SOldCurveTreeContent.h"
+#include "Filters/SCurveEditorFilterPanel.h"
+#include "Filters/ViewModels/Linking/LinkedFilterViewModel.h"
+#include "Filters/Widgets/Linking/SFilterBarSwitcher.h"
+#include "HAL/IConsoleManager.h"
+#include "IPropertyRowGenerator.h"
+#include "IStructureDetailsView.h"
+#include "MVVM/CurveEditorExtension.h"
+#include "SCurveEditorPanel.h"
+#include "SCurveEditorToolProperties.h"
+#include "SCurveKeyDetailPanel.h"
+#include "Sequencer.h"
+#include "SequencerCommands.h"
+#include "Widgets/DeclarativeSyntaxSupport.h"
+#include "Widgets/Docking/SDockTab.h"
+
+struct FTimeSliderArgs;
+
+#define LOCTEXT_NAMESPACE "FSequencerCurveEditorPresenter"
+
+namespace UE::Sequencer
+{
+namespace PresenterDetail
+{
+static TAutoConsoleVariable<bool> CVarUseLinkedFilterUI(
+	TEXT("Sequencer.CurveEditor.UseLinkedFiltering"), false, TEXT("Whether to use the new linked filtering UI in Curve Editor")
+	);
+
+static FCurveEditorWidgets MakePresenterWidgets(
+	const TSharedRef<FSequencer>& InSequencer, 
+	const TSharedRef<FCurveEditor>& InCurveEditor, 
+	const TSharedRef<ITimeSliderController>& InTimeSliderController,
+	const TSharedRef<FLinkedFilterViewModel>& InFilteringViewModel,
+	const FCurveModelSyncer& InCurveModelSyncer UE_LIFETIMEBOUND,
+	const TSharedPtr<FTabManager> InTabManager
+	)
+{
+	const TSharedRef<SOldCurveTreeContent> OldContent = SNew(SOldCurveTreeContent, InSequencer, InCurveEditor)
+		.TabManager(InTabManager);
+	const TSharedRef<SCurveTreeContent> TreeContent = SNew(SCurveTreeContent)
+		.TabManager(InTabManager)
+		.Sequencer(InSequencer)
+		.CurveEditor(InCurveEditor)
+		.FilterViewModel(InFilteringViewModel)
+		.CurveModelSyncer(&InCurveModelSyncer)
+		.FilterPills()
+		[
+			SNew(SFilterBarSwitcher)
+			.LinkedFilterViewModel(InFilteringViewModel)
+		];
+	
+	const TSharedRef<SWidgetSwitcher> Content = 
+		SNew(SWidgetSwitcher)
+		.WidgetIndex_Lambda([]{ return UseLinkedFilteringUI() ? 1 : 0; })
+		+SWidgetSwitcher::Slot()[ OldContent ]
+		+SWidgetSwitcher::Slot()[ TreeContent ];
+	
+	const TSharedRef<SCurveEditorPanel> CurveEditorPanel = SNew(SCurveEditorPanel, InCurveEditor)
+		// Grid lines match the color specified in FSequencerTimeSliderController::OnPaintViewArea
+		.GridLineTint(FLinearColor(0.f, 0.f, 0.f, 0.3f))
+		.ExternalTimeSliderController(InTimeSliderController)
+		.MinimumViewPanelHeight(0.f)
+		.TabManager(InTabManager)
+		.DisabledTimeSnapTooltip(LOCTEXT("CurveEditorTimeSnapDisabledTooltip", "Time Snapping is currently driven by Sequencer."))
+		.TreeContent()
+		[
+			Content
+		];
+	
+	const TSharedRef<SSequencerCurveEditor> RootWidget = SNew(SSequencerCurveEditor, CurveEditorPanel, InSequencer);
+	return FCurveEditorWidgets(RootWidget, CurveEditorPanel, OldContent, TreeContent);
+}
+}
+
+bool UseLinkedFilteringUI()
+{
+	return PresenterDetail::CVarUseLinkedFilterUI.GetValueOnAnyThread();
+}
+
+FCurveEditorWidgetOwner::FCurveEditorWidgetOwner(
+	const TSharedRef<FSequencer>& InSequencer,
+	const TSharedRef<FCurveEditor>& InCurveEditor,
+	const TSharedRef<ITimeSliderController>& InTimeSliderController,
+	const TSharedRef<FLinkedFilterViewModel>& InFilteringViewModel,
+	const FCurveModelSyncer& InCurveModelSyncer
+	)
+	: WeakSequencer(InSequencer)
+	, CurveEditor(InCurveEditor)
+	, Widgets(PresenterDetail::MakePresenterWidgets(
+		InSequencer, InCurveEditor, InTimeSliderController, InFilteringViewModel, InCurveModelSyncer, GetTabManager()
+		))
+{
+	BindCommands();
+	InitFrameNumberPropertyLayout(InSequencer);
+	
+	const TSharedPtr<FTabManager> TabManager = GetTabManager();
+	if (const TSharedPtr<SDockTab> ExistingCurveEditorTab = TabManager->FindExistingLiveTab(FCurveEditorExtension::CurveEditorTabName))
+	{
+		ExistingCurveEditorTab->SetContent(Widgets.RootWidget);
+	}
+}
+
+FCurveEditorWidgetOwner::~FCurveEditorWidgetOwner()
+{
+	// It should not happen defensively prevent crash in case the widget outlives FSequencerCurveEditorPresenter.
+	Widgets.CurveEditorPanel->OnFilterClassChanged.Unbind();
+}
+
+bool FCurveEditorWidgetOwner::IsCurveEditorOpen() const
+{
+	const TSharedPtr<FTabManager> TabManager = GetTabManager();
+	return TabManager && TabManager->FindExistingLiveTab(FCurveEditorExtension::CurveEditorTabName).IsValid();
+}
+
+void FCurveEditorWidgetOwner::OpenCurveEditor()
+{
+	const TSharedPtr<FTabManager> TabManager = GetTabManager();
+	if (!ensure(TabManager))
+	{
+		return;
+	}
+
+	// Request the Tab Manager invoke the tab. This will spawn the tab if needed, otherwise pull it to focus. This assumes
+	// that the Toolkit Host's Tab Manager has already registered a tab with a NullWidget for content.
+	if (const TSharedPtr<SDockTab> CurveEditorTab = TabManager->TryInvokeTab(FCurveEditorExtension::CurveEditorTabName))
+	{
+		CurveEditorTab->SetContent(Widgets.RootWidget);
+
+		const FSlateIcon SequencerGraphIcon = FSlateIcon(FAppStyle::GetAppStyleSetName(), "GenericCurveEditor.TabIcon");
+		CurveEditorTab->SetTabIcon(SequencerGraphIcon.GetIcon());
+		CurveEditorTab->SetLabel(LOCTEXT("SequencerMainGraphEditorTitle", "Sequencer Curves"));
+
+		CurveEditor->ZoomToFit();
+	}
+}
+
+void FCurveEditorWidgetOwner::CloseCurveEditor()
+{
+	const TSharedPtr<FTabManager> TabManager = GetTabManager();
+	if (!TabManager.IsValid())
+	{
+		return;
+	}
+
+	if (const TSharedPtr<SDockTab> CurveEditorTab = TabManager->FindExistingLiveTab(FCurveEditorExtension::CurveEditorTabName))
+	{
+		CurveEditorTab->RequestCloseTab();
+	}
+}
+
+void FCurveEditorWidgetOwner::SyncSelection() const
+{
+	if (UseLinkedFilteringUI())
+	{
+		Widgets.TreeContent->SyncSelection();
+	}
+	else
+	{
+		Widgets.OldTreeContent->SyncSelection();
+	}
+}
+
+TSharedPtr<SCurveEditorTree> FCurveEditorWidgetOwner::GetTreeView() const
+{
+	return UseLinkedFilteringUI() 
+		? Widgets.TreeContent->GetCurveEditorTreeView()
+		: Widgets.OldTreeContent->GetCurveEditorTreeView();
+}
+
+USequencerSettings* FCurveEditorWidgetOwner::GetSequencerSettings() const
+{
+	const TSharedPtr<ISequencer> Sequencer = WeakSequencer.Pin();
+	return Sequencer ? Sequencer->GetSequencerSettings() : nullptr;
+}
+
+TSharedPtr<FTabManager> FCurveEditorWidgetOwner::GetTabManager() const
+{
+	const TSharedPtr<ISequencer> Sequencer = WeakSequencer.Pin();
+	const TSharedPtr<IToolkitHost> ToolkitHost = Sequencer ? Sequencer->GetToolkitHost() : nullptr;
+	return ToolkitHost.IsValid() ? ToolkitHost->GetTabManager() : nullptr;
+}
+
+void FCurveEditorWidgetOwner::BindCommands()
+{
+	const TSharedPtr<FUICommandList> CurveEditorCommandList = CurveEditor->GetCommands();
+	if (!ensure(CurveEditorCommandList))
+	{
+		return;
+	}
+	
+	CurveEditorCommandList->MapAction(
+		FSequencerCommands::Get().QuickTreeSearch,
+		FExecuteAction::CreateRaw(this, &FCurveEditorWidgetOwner::HandleQuickTreeSearch)
+	);
+	CurveEditorCommandList->MapAction(
+		FSequencerCommands::Get().ToggleShowGotoBox,
+		FExecuteAction::CreateRaw(this, &FCurveEditorWidgetOwner::HandleToggleShowGotoBox)
+	);
+}
+
+void FCurveEditorWidgetOwner::HandleQuickTreeSearch()
+{
+	if (UseLinkedFilteringUI())
+	{
+		Widgets.TreeContent->FocusSearchBox();
+	}
+	else
+	{
+		Widgets.OldTreeContent->FocusSearchBox();
+	}
+}
+
+void FCurveEditorWidgetOwner::HandleToggleShowGotoBox()
+{
+	if (UseLinkedFilteringUI())
+	{
+		Widgets.TreeContent->FocusPlayTimeDisplay();
+	}
+	else
+	{
+		Widgets.OldTreeContent->FocusPlayTimeDisplay();
+	}
+}
+
+void FCurveEditorWidgetOwner::InitFrameNumberPropertyLayout(const TSharedRef<FSequencer>& InSequencer)
+{
+	const TSharedRef<SCurveEditorPanel>& CurveEditorPanel = Widgets.CurveEditorPanel;
+	
+	CurveEditorPanel->GetKeyDetailsView()->GetPropertyRowGenerator()->RegisterInstancedCustomPropertyTypeLayout(
+		"FrameNumber", 
+		FOnGetPropertyTypeCustomizationInstance::CreateSP(InSequencer, &FSequencer::MakeFrameNumberDetailsCustomization));
+	CurveEditorPanel->GetToolPropertiesPanel()->GetStructureDetailsView()->GetDetailsView()->RegisterInstancedCustomPropertyTypeLayout(
+		"FrameNumber",
+		FOnGetPropertyTypeCustomizationInstance::CreateSP(InSequencer, &FSequencer::MakeFrameNumberDetailsCustomization));
+	
+	Widgets.CurveEditorPanel->OnFilterClassChanged.BindRaw(this, &FCurveEditorWidgetOwner::OnFilterClassChanged);
+}
+
+void FCurveEditorWidgetOwner::OnFilterClassChanged()
+{
+	const TSharedPtr<FSequencer> Sequencer = WeakSequencer.Pin();
+	const TSharedPtr<SCurveEditorFilterPanel> FilterPanel = Widgets.CurveEditorPanel->GetFilterPanel();
+	if (Sequencer && FilterPanel)
+	{
+		FilterPanel->GetDetailsView()->RegisterInstancedCustomPropertyTypeLayout(
+			TEXT("FrameNumber"),
+			FOnGetPropertyTypeCustomizationInstance::CreateSP(Sequencer.ToSharedRef(), &FSequencer::MakeFrameNumberDetailsCustomization));
+	}
+}
+} // namespace UE::Sequencer
+
+#undef LOCTEXT_NAMESPACE
\ No newline at end of file
diff --git a/Engine/Source/Editor/Sequencer/Private/CurveEditor/Views/CurveEditorWidgetOwner.h b/Engine/Source/Editor/Sequencer/Private/CurveEditor/Views/CurveEditorWidgetOwner.h
new file mode 100644
--- /dev/null
+++ b/Engine/Source/Editor/Sequencer/Private/CurveEditor/Views/CurveEditorWidgetOwner.h
@@ -0,0 +1,87 @@
+﻿// Copyright Epic Games, Inc. All Rights Reserved.
+
+#pragma once
+
+#include "CurveEditorWidgets.h"
+#include "HAL/Platform.h"
+#include "Templates/SharedPointer.h"
+
+class FCurveEditor;
+class FSequencer;
+class FTabManager;
+class ITimeSliderController;
+class SCurveEditorPanel;
+class SCurveEditorTree;
+class SSequencerCurveEditor;
+class USequencerSettings;
+struct FTimeSliderArgs;
+
+namespace UE::Sequencer
+{
+class FCurveModelSyncer;
+class FLinkedFilterViewModel;
+
+/** @return Whether the linked filtering UI is enabled. */
+bool UseLinkedFilteringUI();
+
+/**
+ * Instantiates the UI displaying Sequencer's FCurveEditor.
+ * 
+ * Manages the UI:
+ * - binds UI commands
+ * - updates the details customization for the key settings in the toolbar
+ * - interacts with the tab manager
+ */
+class FCurveEditorWidgetOwner
+{
+public:
+	
+	explicit FCurveEditorWidgetOwner(
+		const TSharedRef<FSequencer>& InSequencer, 
+		const TSharedRef<FCurveEditor>& InCurveEditor,
+		const TSharedRef<ITimeSliderController>& InTimeSliderController,
+		const TSharedRef<FLinkedFilterViewModel>& InFilteringViewModel,
+		const FCurveModelSyncer& InCurveModelSyncer UE_LIFETIMEBOUND
+		);
+	~FCurveEditorWidgetOwner();
+	
+	/** Returns whether the curve editor is open */
+	bool IsCurveEditorOpen() const;
+	
+	/** Opens the curve editor */
+	void OpenCurveEditor();
+	/** Closes the curve editor */
+	void CloseCurveEditor();
+	
+	/** Requests that the UI syncs its selection to that of Sequencer. */
+	void SyncSelection() const;
+	
+	/**
+	 * @return The currently active tree view widget.
+	 * TODO UE-363527: Maybe this can be factored out. It breaks encapsulation. For that, FCurveEditorExtension::GetCurveEditorTreeView must be deprecated.
+	 */
+	TSharedPtr<SCurveEditorTree> GetTreeView() const;
+	
+private:
+	
+	/** Sequencer instance that owns us. */
+	const TWeakPtr<FSequencer> WeakSequencer;
+	/** The curve editor being displayed. */
+	const TSharedRef<FCurveEditor> CurveEditor;
+	
+	/** Widgets this presenter manages. */
+	const FCurveEditorWidgets Widgets;
+	
+	USequencerSettings* GetSequencerSettings() const;
+	TSharedPtr<FTabManager> GetTabManager() const;
+
+	void BindCommands();
+	void HandleQuickTreeSearch();
+	void HandleToggleShowGotoBox();
+	
+	/** Register an instanced custom property type layout to handle converting FFrameNumber from Tick Resolution to Display Rate. */
+	void InitFrameNumberPropertyLayout(const TSharedRef<FSequencer>& InSequencer);
+	/** Updates the frame number layout */
+	void OnFilterClassChanged();
+};
+} // namespace UE::Sequencer
diff --git a/Engine/Source/Editor/Sequencer/Private/CurveEditor/Views/CurveEditorWidgets.h b/Engine/Source/Editor/Sequencer/Private/CurveEditor/Views/CurveEditorWidgets.h
new file mode 100644
--- /dev/null
+++ b/Engine/Source/Editor/Sequencer/Private/CurveEditor/Views/CurveEditorWidgets.h
@@ -0,0 +1,41 @@
+﻿// Copyright Epic Games, Inc. All Rights Reserved.
+
+#pragma once
+
+#include "Templates/SharedPointer.h"
+
+class SCurveEditorPanel;
+class SSequencerCurveEditor;
+
+namespace UE::Sequencer
+{
+class SCurveTreeContent;
+class SOldCurveTreeContent;
+
+/** @see FSequencerCurveEditorPresenter */
+struct FCurveEditorWidgets
+{
+	/** The curve editor widget containing the curve editor panel */
+	TSharedRef<SSequencerCurveEditor> RootWidget;
+	/** The curve editor panel. This is created and updated even if it is not currently visible. */
+	TSharedRef<SCurveEditorPanel> CurveEditorPanel;
+	
+	/** The tree content displayed inside CurveEditorPanel.*/
+	TSharedRef<SOldCurveTreeContent> OldTreeContent;
+	
+	/** The tree content for curve editor. It displays linked filtering UI. */
+	TSharedRef<SCurveTreeContent> TreeContent;
+
+	explicit FCurveEditorWidgets(
+		const TSharedRef<SSequencerCurveEditor>& InRootWidget, 
+		const TSharedRef<SCurveEditorPanel>& InCurveEditorPanel,
+		const TSharedRef<SOldCurveTreeContent>& InOldTreeContent,
+		const TSharedRef<SCurveTreeContent>& InTreeContent
+		)
+		: RootWidget(InRootWidget)
+		, CurveEditorPanel(InCurveEditorPanel)
+		, OldTreeContent(InOldTreeContent)
+		, TreeContent(InTreeContent)
+	{}
+};
+} // namespace UE::Sequencer
diff --git a/Engine/Source/Editor/Sequencer/Private/CurveEditor/Views/SelectionSyncUtils.cpp b/Engine/Source/Editor/Sequencer/Private/CurveEditor/Views/SelectionSyncUtils.cpp
new file mode 100644
--- /dev/null
+++ b/Engine/Source/Editor/Sequencer/Private/CurveEditor/Views/SelectionSyncUtils.cpp
@@ -0,0 +1,197 @@
+﻿// Copyright Epic Games, Inc. All Rights Reserved.
+
+#include "SelectionSyncUtils.h"
+
+#include "CurveEditor.h"
+#include "CurveEditor/Integration/Selection/CurveModelSyncer.h"
+#include "MVVM/Extensions/ICurveEditorTreeItemExtension.h"
+#include "MVVM/Extensions/IOutlinerExtension.h"
+#include "MVVM/Selection/Selection.h"
+#include "MVVM/ViewModelPtr.h"
+#include "Misc/ScopeExit.h"
+#include "Templates/SharedPointer.h"
+#include "Tree/SCurveEditorTree.h"
+
+namespace UE::Sequencer
+{
+namespace SyncDetail
+{
+static void GetAllChildrenRecursively(
+	const TSharedRef<SCurveEditorTree>& CurveEditorTreeView, const FCurveEditorTreeItemID Item, int32& LastAddedItem,
+	TMap<const FCurveEditorTreeItemID, int32>& ItemToSortOrder
+	)
+{
+	ItemToSortOrder.Add(Item, LastAddedItem++);
+
+	TArray<FCurveEditorTreeItemID> Children;
+
+	CurveEditorTreeView->GetTreeItemChildren(Item, Children);
+	// Add them to the output array
+	for (const FCurveEditorTreeItemID ChildItemID : Children)
+	{
+		GetAllChildrenRecursively(CurveEditorTreeView, ChildItemID, LastAddedItem, ItemToSortOrder);
+	}
+}
+} // namespace SyncDetail
+
+void ApplySequencerToCurveEditorSelection(
+	const TSharedPtr<FSequencerSelection>& InSelection, 
+	const TSharedRef<FCurveEditor>& InCurveEditorModel, 
+	const TSharedRef<SCurveEditorTree>& InCurveEditorTreeView
+	)
+{
+	//we now sync to the first non-channel in the tree hierarchy, this is needed for hosted channels with control rigs
+	//we need to sort the hierarchy each time since the curve id's being used previously were not consistent in order of viewing/creation, persistent now for undo/redo
+	InCurveEditorModel->SuspendBroadcast();
+	ON_SCOPE_EXIT{ InCurveEditorModel->ResumeBroadcast(); };
+	InCurveEditorTreeView->ClearSelection();
+
+	//map of tree items id's to their actually sort order, need to use this since tree items ids are not in order anymore
+	TMap<const FCurveEditorTreeItemID, int32>  ItemToSortOrder;
+	int32 LastAddedItem = 0;
+	TArray<FCurveEditorTreeItemID> Children;
+	TArrayView<const FCurveEditorTreeItemID> RootItems = InCurveEditorTreeView->GetRootItems();
+	
+	for (const FCurveEditorTreeItemID RootItemID : RootItems)
+	{
+		SyncDetail::GetAllChildrenRecursively(InCurveEditorTreeView, RootItemID, LastAddedItem, ItemToSortOrder);
+	}
+
+	auto GetSortOrder = [&ItemToSortOrder](FCurveEditorTreeItemID InTreeItemID) -> int32
+	{
+		if (const int32* SortOrder = ItemToSortOrder.Find(InTreeItemID))
+		{
+			return *SortOrder;
+		}
+		return INDEX_NONE;
+	};
+
+	FCurveEditorTreeItemID FirstCurveEditorTreeItemID;
+	int32 FirstSortOrder = INDEX_NONE;
+	bool bSelectedIsChannel = true;
+
+	for (TViewModelPtr<IOutlinerExtension> SelectedItem : InSelection->Outliner)
+	{
+		if (TViewModelPtr<ICurveEditorTreeItemExtension> CurveEditorItem = SelectedItem.ImplicitCast())
+		{
+			FCurveEditorTreeItemID CurveEditorTreeItem = CurveEditorItem->GetCurveEditorItemID();
+			if (CurveEditorTreeItem != FCurveEditorTreeItemID::Invalid())
+			{
+				if (!InCurveEditorTreeView->IsItemSelected(CurveEditorTreeItem))
+				{
+					InCurveEditorTreeView->SetItemSelection(CurveEditorTreeItem, true);
+					bool bIsChannel = SelectedItem.AsModel()->IsA<FChannelModel>() || SelectedItem.AsModel()->IsA<FChannelGroupOutlinerModel>()
+					|| SelectedItem.AsModel()->IsA<FChannelGroupModel>();
+					if (!FirstCurveEditorTreeItemID.IsValid())
+					{
+						FirstCurveEditorTreeItemID = CurveEditorTreeItem;
+						FirstSortOrder = GetSortOrder(CurveEditorTreeItem);
+						bSelectedIsChannel = bIsChannel;
+					}
+					else
+					{
+						int32 SortOrder = GetSortOrder(CurveEditorTreeItem);
+						if (bIsChannel == false)
+						{
+							if (bSelectedIsChannel == true)
+							{
+								FirstCurveEditorTreeItemID = CurveEditorTreeItem;
+								FirstSortOrder = SortOrder;
+								bSelectedIsChannel = false;
+							}
+							else if (SortOrder < FirstSortOrder)
+							{
+								FirstCurveEditorTreeItemID = CurveEditorTreeItem;
+								FirstSortOrder = SortOrder;
+								bSelectedIsChannel = false;
+							}
+						}
+						else if (bSelectedIsChannel == true) //its a channel only sort if selected is a non channel
+						{
+							if (SortOrder < FirstSortOrder)
+							{
+								FirstCurveEditorTreeItemID = CurveEditorTreeItem;
+								FirstSortOrder = SortOrder;
+							}
+						}
+					}
+				}
+			}
+		}
+	}
+	
+	if (FirstCurveEditorTreeItemID.IsValid())
+	{
+		InCurveEditorTreeView->RequestScrollIntoView(FirstCurveEditorTreeItemID);
+	}
+}
+
+namespace SyncDetail
+{
+static void SetParentsExpanded(const FCurveEditorTree& InTree, const FCurveEditorTreeItemID& InItem, SCurveEditorTree& InTreeView, bool bInExpanded)
+{
+	for (const FCurveEditorTreeItem* Item = InTree.FindItem(InItem); Item; Item = InTree.FindItem(Item->GetParentID()))
+	{
+		InTreeView.SetItemExpansion(Item->GetID(), bInExpanded);
+	}
+}
+}
+
+void ExpandSelectedCurveEditorItems(FCurveEditorTree& InTree, SCurveEditorTree& InTreeView)
+{
+	for (const TPair<FCurveEditorTreeItemID, ECurveEditorTreeSelectionState>& Pair : InTree.GetSelection())
+	{
+		if (Pair.Value != ECurveEditorTreeSelectionState::None)
+		{
+			SyncDetail::SetParentsExpanded(InTree, Pair.Key, InTreeView, true);
+		}
+	}
+}
+
+namespace SyncDetail
+{
+static void SetParentsExpanded(const TViewModelPtr<IOutlinerExtension>& InNode, const bool bInExpanded)
+{
+	for (const TViewModelPtr<IOutlinerExtension>& ParentNode : InNode.AsModel()->GetAncestorsOfType<IOutlinerExtension>())
+	{
+		if (ParentNode->IsExpanded() != bInExpanded)
+		{
+			ParentNode->SetExpansion(bInExpanded);
+		}
+	}
+}
+}
+
+void ApplyCurveEditorToSequencer(FSequencerSelection& InSelection, const FCurveEditorTree& InTree, const FCurveModelSyncer& InCurveModelSyncer)
+{
+	FOutlinerSelection NewOutlinerSelection;
+	const TMap<FCurveEditorTreeItemID, ECurveEditorTreeSelectionState>& CurveSelection = InTree.GetSelection();
+	for (const TPair<FCurveEditorTreeItemID, ECurveEditorTreeSelectionState>& Pair : CurveSelection)
+	{
+		if (Pair.Value != ECurveEditorTreeSelectionState::Explicit)
+		{
+			continue;
+		}
+		
+		const TSharedPtr<FViewModel> ViewModel = InCurveModelSyncer.FindViewModelById(Pair.Key).Pin();
+		const TViewModelPtr<IOutlinerExtension> OutlinerNode = ViewModel ? CastViewModel<IOutlinerExtension>(ViewModel.ToSharedRef()) : nullptr;
+		if (OutlinerNode)
+		{
+			NewOutlinerSelection.Select(OutlinerNode);
+		}
+	}
+	
+	InSelection.Outliner.ReplaceWith(MoveTemp(NewOutlinerSelection));
+}
+
+void ExpandSelectedSequencerItems(const FSequencerSelection& InSelection)
+{
+	for (const TWeakViewModelPtr<IOutlinerExtension>& WeakOutlinerNode : InSelection.Outliner.GetSelected())
+	{
+		if (const TViewModelPtr<IOutlinerExtension> OutlinerNode = WeakOutlinerNode.Pin())
+		{
+			SyncDetail::SetParentsExpanded(OutlinerNode, true);
+		}
+	}
+}
+} // namespace UE::Sequencer
diff --git a/Engine/Source/Editor/Sequencer/Private/CurveEditor/Views/SelectionSyncUtils.h b/Engine/Source/Editor/Sequencer/Private/CurveEditor/Views/SelectionSyncUtils.h
new file mode 100644
--- /dev/null
+++ b/Engine/Source/Editor/Sequencer/Private/CurveEditor/Views/SelectionSyncUtils.h
@@ -0,0 +1,31 @@
+﻿// Copyright Epic Games, Inc. All Rights Reserved.
+
+#pragma once
+
+#include "Templates/SharedPointerFwd.h"
+
+class FCurveEditorTree;
+class FCurveEditor;
+class SCurveEditorTree;
+
+namespace UE::Sequencer
+{
+class FCurveModelSyncer;
+class FSequencerSelection;
+
+/** Applies selection in Sequencer to Curve Editor. */
+void ApplySequencerToCurveEditorSelection(
+	const TSharedPtr<FSequencerSelection>& InSelection, 
+	const TSharedRef<FCurveEditor>& InCurveEditorModel, 
+	const TSharedRef<SCurveEditorTree>& InCurveEditorTreeView
+	);
+
+/** Expands the tree hierarchy to show all items selected in Curve Editor. */
+void ExpandSelectedCurveEditorItems(FCurveEditorTree& InTree, SCurveEditorTree& InTreeView);
+
+/** Applies selection in Curve Editor to Sequencer.*/
+void ApplyCurveEditorToSequencer(FSequencerSelection& InSelection, const FCurveEditorTree& InTree, const FCurveModelSyncer& InCurveModelSyncer);
+
+/** Expands the tree hierarchy to show all items selected in Sequencer. */
+void ExpandSelectedSequencerItems(const FSequencerSelection& InSelection);
+} // namespace UE::Sequencer
diff --git a/Engine/Source/Editor/Sequencer/Private/CurveEditor/Widgets/SCurveEditorSearchAndFilterRow.cpp b/Engine/Source/Editor/Sequencer/Private/CurveEditor/Widgets/SCurveEditorSearchAndFilterRow.cpp
new file mode 100644
--- /dev/null
+++ b/Engine/Source/Editor/Sequencer/Private/CurveEditor/Widgets/SCurveEditorSearchAndFilterRow.cpp
@@ -0,0 +1,102 @@
+﻿// Copyright Epic Games, Inc. All Rights Reserved.
+
+#include "SCurveEditorSearchAndFilterRow.h"
+
+#include "Filters/Menus/SequencerTrackFilterMenu.h"
+#include "Filters/SequencerFilterBar.h"
+#include "Filters/Utils/HideIsolateViewModelUtils.h"
+#include "Filters/ViewModels/HideIsolateShowViewModel.h"
+#include "Filters/ViewModels/Linking/ILinkedFilterViewModel.h"
+#include "Filters/ViewModels/Linking/LinkedFilterViewModel.h"
+#include "Filters/Widgets/Linking/SFilterLinkStateSwitcher.h"
+#include "Filters/Widgets/Linking/SLinkedFilteringToggleButton.h"
+#include "Filters/Widgets/SSearchAndFilterWidget.h"
+#include "MVVM/Selection/Selection.h"
+#include "Sequencer.h"
+#include "Filters/Utils/LinkedFilteringViewUtils.h"
+#include "Widgets/SBoxPanel.h"
+#include "Widgets/Sequencer/SSequencerViewOptionsButton.h"
+
+namespace UE::Sequencer
+{
+namespace CurveEditorRowDetail
+{
+using FOutlinerExtensionsAttr = TAttribute<TSet<TWeakViewModelPtr<IOutlinerExtension>>>; 
+
+static FOnGetContent MakeDelegate_AddFilterMenuContent(const TSharedRef<FSequencerFilterBar>& InFilterBar)
+{
+	return FOnGetContent::CreateLambda([Menu = MakeShared<FSequencerTrackFilterMenu>(), WeakFilterBar = InFilterBar.ToWeakPtr()]
+	{
+		const TSharedPtr<FSequencerFilterBar> FilterBar = WeakFilterBar.Pin();
+		return ensure(FilterBar) ? Menu->CreateMenu(FilterBar.ToSharedRef()) : SNullWidget::NullWidget;
+	});
+}
+}
+
+void SCurveEditorSearchAndFilterRow::Construct(const FArguments& InArgs)
+{
+	WeakSequencer = InArgs._Sequencer;
+	FilterModel = InArgs._FilterViewModel;
+	
+	ChildSlot
+	[
+		MakeFilterContent(InArgs)
+	];
+}
+
+TSharedPtr<SSequencerSearchBox> SCurveEditorSearchAndFilterRow::GetActiveSearchBox() const
+{
+	static_assert(static_cast<int32>(ELinkedFilterMode::Count) == 2, "Update this switch");
+	switch (FilterModel->GetFilterMode())
+	{
+	case ELinkedFilterMode::Linked: return LinkedFilterWidgets->GetSearchBox();
+	case ELinkedFilterMode::Instanced: return InstancedFilterWidgets->GetSearchBox();
+	default: checkNoEntry(); return nullptr;
+	}
+}
+
+TSharedRef<SWidget> SCurveEditorSearchAndFilterRow::MakeFilterContent(const FArguments& InArgs)
+{
+	const auto MakeWidget = [this, &InArgs](const TSharedRef<FSequencerFilterBar>& InFilterBar)
+	{
+		const TSharedRef<FSequencer> SequencerPin = WeakSequencer.Pin().ToSharedRef();
+		const TSharedRef<FHideIsolateShowViewModel> IsolateViewModel = MakeShared<FHideIsolateShowViewModel>(
+			InFilterBar, InFilterBar->GetHideIsolateFilter(),
+			InArgs._SelectedItems, InArgs._AllItems
+			);
+		IsolateViewModel->OnRequestScrollToView().Add(InArgs._ScrollItemIntoView);
+			
+		const TSharedPtr<ISequencerTrackFilters> Filter = StaticCastSharedPtr<ISequencerTrackFilters>(InFilterBar.ToSharedPtr());
+		return SNew(SSearchAndFilterWidget, InFilterBar, IsolateViewModel)
+			.OnSearchTextChanged(this, &SCurveEditorSearchAndFilterRow::OnSearchTextChanged, Filter)
+			.OnMakeAddFilterMenuContent(CurveEditorRowDetail::MakeDelegate_AddFilterMenuContent(InFilterBar))
+			.BeforeSearchBox()
+			[
+				SNew(SLinkedFilteringToggleButton, InArgs._FilterViewModel.ToSharedRef())
+				.Visibility_Lambda([]
+				{
+					return IsLinkedFilteringFeatureEnabled() ? EVisibility::Visible : EVisibility::Collapsed;
+				})
+			];
+	};
+	
+	const TSharedPtr<FLinkedFilterViewModel> LinkedFilterViewModel = InArgs._FilterViewModel;
+	LinkedFilterWidgets = MakeWidget(LinkedFilterViewModel->GetLinkedFilterBarImpl());
+	InstancedFilterWidgets = MakeWidget(LinkedFilterViewModel->GetInstancedFilterBarImpl());
+	return SNew(SFilterLinkStateSwitcher)
+		.LinkedFilterViewModel(InArgs._FilterViewModel)
+		.LinkedContent() 
+		[
+			LinkedFilterWidgets.ToSharedRef()
+		]
+		.InstancedContent() 
+		[
+			InstancedFilterWidgets.ToSharedRef()
+		];
+}
+
+void SCurveEditorSearchAndFilterRow::OnSearchTextChanged(const FText& InNewText, TSharedPtr<ISequencerTrackFilters> InFilter)
+{
+	InFilter->SetTextFilterString(InNewText.ToString());
+}
+} // namespace UE::Sequencer
diff --git a/Engine/Source/Editor/Sequencer/Private/CurveEditor/Widgets/SCurveEditorSearchAndFilterRow.h b/Engine/Source/Editor/Sequencer/Private/CurveEditor/Widgets/SCurveEditorSearchAndFilterRow.h
new file mode 100644
--- /dev/null
+++ b/Engine/Source/Editor/Sequencer/Private/CurveEditor/Widgets/SCurveEditorSearchAndFilterRow.h
@@ -0,0 +1,65 @@
+﻿// Copyright Epic Games, Inc. All Rights Reserved.
+
+#pragma once
+
+#include "MVVM/Extensions/IOutlinerExtension.h"
+#include "Widgets/DeclarativeSyntaxSupport.h"
+#include "Widgets/SCompoundWidget.h"
+
+class SSequencerSearchBox;
+class ISequencerTrackFilters;
+class FSequencer;
+
+namespace UE::Sequencer
+{
+class FLinkedFilterViewModel;
+class ILinkedFilterViewModel;
+class SSearchAndFilterWidget;
+
+DECLARE_DELEGATE_OneParam(FScrollIntoView, const TWeakViewModelPtr<IOutlinerExtension>&);
+
+/** The row of widgets displayed above the Curve Editor tree view. */
+class SCurveEditorSearchAndFilterRow : public SCompoundWidget
+{
+public:
+	
+	SLATE_BEGIN_ARGS(SCurveEditorSearchAndFilterRow) {}
+		/** Used to generate the subwidgets */
+		SLATE_ARGUMENT(TWeakPtr<FSequencer>, Sequencer)
+		/** Determines the filtering widgets displayed. */
+		SLATE_ARGUMENT(TSharedPtr<FLinkedFilterViewModel>, FilterViewModel)
+
+		/** Gets all items displayed in the tree view. */
+		SLATE_ATTRIBUTE(TSet<TWeakViewModelPtr<IOutlinerExtension>>, AllItems)
+		/** Gets the items selected in the tree view. */
+		SLATE_ATTRIBUTE(TSet<TWeakViewModelPtr<IOutlinerExtension>>, SelectedItems)
+		/** Requests that the given item is scrolled into view. */
+		SLATE_EVENT(FScrollIntoView, ScrollItemIntoView)
+	SLATE_END_ARGS()
+	
+	void Construct(const FArguments& InArgs);
+	
+	/** @return The search box that is currently being displayed. */
+	TSharedPtr<SSequencerSearchBox> GetActiveSearchBox() const;
+
+private:
+	
+	/** Used to update Sequencer's filters. */
+	TWeakPtr<FSequencer> WeakSequencer;
+	
+	/** Filter model for the Curve Editor. */
+	TSharedPtr<ILinkedFilterViewModel> FilterModel;
+	
+	/** Content displayed when ELinkedFilterMode::Linked */
+	TSharedPtr<SSearchAndFilterWidget> LinkedFilterWidgets;
+	/** Content displayed when ELinkedFilterMode::Instanced */
+	TSharedPtr<SSearchAndFilterWidget> InstancedFilterWidgets;
+
+	/** @return The filtering widget. */
+	TSharedRef<SWidget> MakeFilterContent(const FArguments& InArgs);
+	
+	void OnSearchTextChanged(const FText& InNewText, TSharedPtr<ISequencerTrackFilters> InFilter);
+};
+} // namespace UE::Sequencer
+
+
diff --git a/Engine/Source/Editor/Sequencer/Private/CurveEditor/Widgets/SSequencerCurveEditor.cpp b/Engine/Source/Editor/Sequencer/Private/CurveEditor/Widgets/SSequencerCurveEditor.cpp
new file mode 100644
--- /dev/null
+++ b/Engine/Source/Editor/Sequencer/Private/CurveEditor/Widgets/SSequencerCurveEditor.cpp
@@ -0,0 +1,75 @@
+// Copyright Epic Games, Inc. All Rights Reserved.
+
+#include "SSequencerCurveEditor.h"
+#include "Framework/MultiBox/MultiBoxBuilder.h"
+#include "MovieScene.h"
+#include "MovieSceneCommonHelpers.h"
+#include "SCurveEditorPanel.h"
+#include "Sequencer.h"
+#include "Misc/ResizeParamUtils.h"
+
+#define LOCTEXT_NAMESPACE "SSequencerCurveEditor"
+
+void SSequencerCurveEditor::Construct(const FArguments& InArgs, TSharedRef<SCurveEditorPanel> InEditorPanel, TSharedPtr<FSequencer> InSequencer)
+{
+	WeakSequencer = InSequencer;
+
+	ChildSlot
+	[
+		SNew(SVerticalBox)
+
+		+ SVerticalBox::Slot()
+		.AutoHeight()
+		[
+			MakeToolbar(InEditorPanel)
+		]
+		+ SVerticalBox::Slot()
+		.FillHeight(1.0f)
+		[
+			InEditorPanel
+		]
+	];
+}
+
+TSharedRef<SWidget> SSequencerCurveEditor::MakeToolbar(TSharedRef<SCurveEditorPanel> InEditorPanel)
+{
+	FSlimHorizontalToolBarBuilder ToolBarBuilder(InEditorPanel->GetCommands(), FMultiBoxCustomization::None, InEditorPanel->GetToolbarExtender(), true);
+	ToolBarBuilder.BeginSection("Asset");
+
+	{
+		TAttribute<FSlateIcon> SaveIcon;
+		SaveIcon.Bind(TAttribute<FSlateIcon>::FGetter::CreateLambda([&] {
+
+			TArray<UMovieScene*> MovieScenesToSave;
+			MovieSceneHelpers::GetDescendantMovieScenes(WeakSequencer.Pin()->GetRootMovieSceneSequence(), MovieScenesToSave);
+			for (UMovieScene* MovieSceneToSave : MovieScenesToSave)
+			{
+				UPackage* MovieScenePackageToSave = MovieSceneToSave->GetOuter()->GetOutermost();
+				if (MovieScenePackageToSave->IsDirty())
+				{
+					return FSlateIcon(FAppStyle::Get().GetStyleSetName(), "Icons.SaveChanged");
+				}
+			}
+
+			return FSlateIcon(FAppStyle::Get().GetStyleSetName(), "Icons.Save");
+			}));
+
+		if (WeakSequencer.Pin()->GetHostCapabilities().bSupportsSaveMovieSceneAsset)
+		{
+			ToolBarBuilder.AddToolBarButton(
+				FUIAction(FExecuteAction::CreateLambda([&] { WeakSequencer.Pin()->SaveCurrentMovieScene(); })),
+				NAME_None,
+				LOCTEXT("SaveDirtyPackages", "Save"),
+				LOCTEXT("SaveDirtyPackagesTooltip", "Saves the current sequence and any subsequences"),
+				SaveIcon,
+				EUserInterfaceActionType::Button, NAME_None, {}, {}, UE::CurveEditor::MakeResizeParams(TEXT("Sequencer.Save"))
+			);
+		}		
+	}
+	
+	ToolBarBuilder.EndSection();
+	// We just use all of the extenders as our toolbar, we don't have a need to create a separate toolbar.
+	return ToolBarBuilder.MakeWidget();
+}
+
+#undef LOCTEXT_NAMESPACE
diff --git a/Engine/Source/Editor/Sequencer/Private/CurveEditor/Widgets/SSequencerCurveEditor.h b/Engine/Source/Editor/Sequencer/Private/CurveEditor/Widgets/SSequencerCurveEditor.h
new file mode 100644
--- /dev/null
+++ b/Engine/Source/Editor/Sequencer/Private/CurveEditor/Widgets/SSequencerCurveEditor.h
@@ -0,0 +1,23 @@
+// Copyright Epic Games, Inc. All Rights Reserved.
+
+#pragma once
+
+#include "Widgets/SCompoundWidget.h"
+
+class FSequencer;
+class SCurveEditorPanel;
+
+class SSequencerCurveEditor : public SCompoundWidget
+{
+	SLATE_BEGIN_ARGS(SSequencerCurveEditor)
+	{}
+	SLATE_END_ARGS()
+public:
+	void Construct(const FArguments& InArgs, TSharedRef<SCurveEditorPanel> InEditorPanel, TSharedPtr<FSequencer> InSequencer);
+
+	TSharedRef<SWidget> MakeToolbar(TSharedRef<SCurveEditorPanel> InEditorPanel);
+
+private:
+	TWeakPtr<FSequencer> WeakSequencer;
+};
+
diff --git a/Engine/Source/Editor/Sequencer/Private/CurveEditor/Widgets/SequencerCurveEditorTimeSliderController.cpp b/Engine/Source/Editor/Sequencer/Private/CurveEditor/Widgets/SequencerCurveEditorTimeSliderController.cpp
new file mode 100644
--- /dev/null
+++ b/Engine/Source/Editor/Sequencer/Private/CurveEditor/Widgets/SequencerCurveEditorTimeSliderController.cpp
@@ -0,0 +1,129 @@
+// Copyright Epic Games, Inc. All Rights Reserved.
+
+
+#include "SequencerCurveEditorTimeSliderController.h"
+
+#include "CurveEditor.h"
+#include "Sequencer.h"
+#include "SequencerSettings.h"
+
+FSequencerCurveEditorTimeSliderController::FSequencerCurveEditorTimeSliderController(const FTimeSliderArgs& InArgs, TWeakPtr<FSequencer> InWeakSequencer, TSharedRef<FCurveEditor> InCurveEditor)
+	: FSequencerTimeSliderController(InArgs, InWeakSequencer)
+{
+	WeakSequencer = InWeakSequencer;
+	WeakCurveEditor = InCurveEditor;
+}
+
+
+void FSequencerCurveEditorTimeSliderController::ClampViewRange(double& NewRangeMin, double& NewRangeMax)
+{
+	// Since the CurveEditor uses a different view range (potentially) we have to be careful about which one we clamp.
+	TSharedPtr<ISequencer> Sequencer = WeakSequencer.Pin();
+	if (!Sequencer.IsValid())
+	{
+		return;
+	}
+
+	const bool bLinkedTimeRange = Sequencer->GetSequencerSettings()->GetLinkCurveEditorTimeRange();
+	if (bLinkedTimeRange)
+	{
+		return FSequencerTimeSliderController::ClampViewRange(NewRangeMin, NewRangeMax);
+	}
+	else
+	{
+		TSharedPtr<FCurveEditor> CurveEditor = WeakCurveEditor.Pin();
+		if (CurveEditor.IsValid())
+		{
+			double InputMin, InputMax;
+			CurveEditor->GetBounds().GetInputBounds(InputMin, InputMax);
+
+			bool bNeedsClampSet = false;
+			double NewClampRangeMin = InputMin;
+			if (NewRangeMin < InputMin)
+			{
+				NewClampRangeMin = NewRangeMin;
+				bNeedsClampSet = true;
+			}
+
+			double NewClampRangeMax = InputMax;
+			if (NewRangeMax > InputMax)
+			{
+				NewClampRangeMax = NewRangeMax;
+				bNeedsClampSet = true;
+			}
+
+			if (bNeedsClampSet)
+			{
+				CurveEditor->GetBounds().SetInputBounds(NewClampRangeMin, NewClampRangeMax);
+			}
+
+		}
+	}
+}
+
+void FSequencerCurveEditorTimeSliderController::SetViewRange(double NewRangeMin, double NewRangeMax, EViewRangeInterpolation Interpolation)
+{
+	TSharedPtr<ISequencer> Sequencer = WeakSequencer.Pin();
+	if (!Sequencer.IsValid())
+	{
+		return;
+	}
+
+	const bool bLinkedTimeRange = Sequencer->GetSequencerSettings()->GetLinkCurveEditorTimeRange();
+	if (bLinkedTimeRange)
+	{
+		return FSequencerTimeSliderController::SetViewRange(NewRangeMin, NewRangeMax, Interpolation);
+	}
+	else
+	{
+		// Clamp to a minimum size to avoid zero-sized or negative visible ranges
+		double MinVisibleTimeRange = FFrameNumber(1) / GetTickResolution();
+		TRange<double> ExistingViewRange = GetViewRange();
+
+		if (NewRangeMax == ExistingViewRange.GetUpperBoundValue())
+		{
+			if (NewRangeMin > NewRangeMax - MinVisibleTimeRange)
+			{
+				NewRangeMin = NewRangeMax - MinVisibleTimeRange;
+			}
+		}
+		else if (NewRangeMax < NewRangeMin + MinVisibleTimeRange)
+		{
+			NewRangeMax = NewRangeMin + MinVisibleTimeRange;
+		}
+
+		TSharedPtr<FCurveEditor> CurveEditor = WeakCurveEditor.Pin();
+		if (CurveEditor.IsValid())
+		{
+			CurveEditor->GetBounds().SetInputBounds(NewRangeMin, NewRangeMax);
+		}
+	}
+}
+
+FAnimatedRange FSequencerCurveEditorTimeSliderController::GetViewRange() const
+{ 
+	// If they've linked the Sequencer timerange we can return the internal controller's view range, otherwise we return the bounds (which internally does the same check)
+	TSharedPtr<ISequencer> Sequencer = WeakSequencer.Pin();
+	if (!Sequencer.IsValid())
+	{
+		return FAnimatedRange();
+	}
+	const bool bLinkedTimeRange = Sequencer->GetSequencerSettings()->GetLinkCurveEditorTimeRange();
+	if (bLinkedTimeRange)
+	{
+		return FSequencerTimeSliderController::GetViewRange();
+	}
+	else
+	{
+		TSharedPtr<FCurveEditor> CurveEditor = WeakCurveEditor.Pin();
+		if (CurveEditor.IsValid())
+		{
+			double InputMin, InputMax;
+			CurveEditor->GetBounds().GetInputBounds(InputMin, InputMax);
+
+			return FAnimatedRange(InputMin, InputMax);
+		}
+	}
+
+	return FAnimatedRange();
+}
diff --git a/Engine/Source/Editor/Sequencer/Private/CurveEditor/Widgets/SequencerCurveEditorTimeSliderController.h b/Engine/Source/Editor/Sequencer/Private/CurveEditor/Widgets/SequencerCurveEditorTimeSliderController.h
new file mode 100644
--- /dev/null
+++ b/Engine/Source/Editor/Sequencer/Private/CurveEditor/Widgets/SequencerCurveEditorTimeSliderController.h
@@ -0,0 +1,25 @@
+// Copyright Epic Games, Inc. All Rights Reserved.
+
+#pragma once
+
+#include "CoreMinimal.h"
+#include "SequencerTimeSliderController.h"
+
+class ISequencer;
+class FSequencer;
+
+class FSequencerCurveEditorTimeSliderController : public FSequencerTimeSliderController
+{
+public:
+
+	FSequencerCurveEditorTimeSliderController(const FTimeSliderArgs& InArgs, TWeakPtr<FSequencer> InWeakSequencer, TSharedRef<FCurveEditor> InCurveEditor);
+
+	virtual void ClampViewRange(double& NewRangeMin, double& NewRangeMax) override;
+	virtual void SetViewRange(double NewRangeMin, double NewRangeMax, EViewRangeInterpolation Interpolation) override;
+	virtual FAnimatedRange GetViewRange() const override;
+
+private:
+	TWeakPtr<ISequencer> WeakSequencer;
+	TWeakPtr<FCurveEditor> WeakCurveEditor;
+};
+
diff --git a/Engine/Source/Editor/Sequencer/Private/CurveEditor/Widgets/Tree/SCurveTreeContent.cpp b/Engine/Source/Editor/Sequencer/Private/CurveEditor/Widgets/Tree/SCurveTreeContent.cpp
new file mode 100644
--- /dev/null
+++ b/Engine/Source/Editor/Sequencer/Private/CurveEditor/Widgets/Tree/SCurveTreeContent.cpp
@@ -0,0 +1,217 @@
+﻿// Copyright Epic Games, Inc. All Rights Reserved.
+
+#include "SCurveTreeContent.h"
+
+#include "Components/VerticalBox.h"
+#include "CurveEditor.h"
+#include "CurveEditor/Integration/Selection/CurveModelSyncer.h"
+#include "CurveEditor/Views/SelectionSyncUtils.h"
+#include "CurveEditor/Widgets/SCurveEditorSearchAndFilterRow.h"
+#include "Filters/Widgets/SSequencerSearchBox.h"
+#include "Misc/TransportControlsHelper.h"
+#include "SPopoutTabInlineContent.h"
+#include "STemporarilyFocusedSpinBox.h"
+#include "Sequencer.h"
+#include "Styling/AppStyle.h"
+#include "Tree/SCurveEditorTree.h"
+#include "Widgets/Layout/SBorder.h"
+#include "Widgets/Layout/SScrollBorder.h"
+
+namespace UE::Sequencer
+{
+void SCurveTreeContent::Construct(const FArguments& InArgs)
+{
+	const bool bHasCurveModelSyncer = InArgs._CurveModelSyncer.IsBound() || InArgs._CurveModelSyncer.IsSet();
+	check(InArgs._Sequencer && InArgs._CurveEditor && InArgs._FilterViewModel && bHasCurveModelSyncer);
+	
+	WeakSequencer = InArgs._Sequencer;
+	CurveEditor = InArgs._CurveEditor;
+	CurveModelSyncerAttr = InArgs._CurveModelSyncer;
+	
+	ChildSlot
+	[
+		SNew(SVerticalBox)
+
+		+ SVerticalBox::Slot()
+		.AutoHeight()
+		[
+			MakeFilterArea(InArgs)
+		]
+
+		+SVerticalBox::Slot()
+		[
+			MakeTreeContent(InArgs)
+		]
+
+		+ SVerticalBox::Slot()
+		.AutoHeight()
+		[
+			MakeFooter(InArgs)
+		]
+	];
+}
+
+void SCurveTreeContent::FocusSearchBox()
+{
+	FSlateApplication::Get().SetKeyboardFocus(FilterRow->GetActiveSearchBox(), EFocusCause::SetDirectly);
+}
+
+void SCurveTreeContent::FocusPlayTimeDisplay()
+{
+	PlayTimeDisplay->Setup(); 
+	FSlateApplication::Get().SetKeyboardFocus(PlayTimeDisplay, EFocusCause::SetDirectly);
+}
+
+void SCurveTreeContent::SyncSelection() const
+{
+	if (const TSharedPtr<FSequencer> Sequencer = WeakSequencer.Pin())
+	{
+		Sequencer::ApplySequencerToCurveEditorSelection(
+			Sequencer->GetViewModel()->GetSelection(), CurveEditor.ToSharedRef(), CurveEditorTreeView.ToSharedRef()
+			);
+		ExpandSelectedCurveEditorItems(*CurveEditor->GetTree(), *CurveEditorTreeView.ToSharedRef());
+	}
+}
+
+TSharedRef<SWidget> SCurveTreeContent::MakeFilterArea(const FArguments& InArgs)
+{
+	return SNew(SVerticalBox)
+		
+		+SVerticalBox::Slot()
+		.AutoHeight()
+		[
+			SNew(SBorder)
+			.BorderImage(FAppStyle::GetBrush("ToolPanel.GroupBorder"))
+			.Clipping(EWidgetClipping::ClipToBounds)
+			[
+				SAssignNew(FilterRow, SCurveEditorSearchAndFilterRow)
+				.Sequencer(InArgs._Sequencer)
+				.FilterViewModel(InArgs._FilterViewModel)
+				.ScrollItemIntoView(this, &SCurveTreeContent::ScrollIntoView)
+				.AllItems(this, &SCurveTreeContent::GetAllItems)
+				.SelectedItems(this, &SCurveTreeContent::GetSelectedItems)
+			]
+		]
+		
+		+SVerticalBox::Slot()
+		.AutoHeight()
+		[
+			InArgs._FilterPills.Widget
+		];
+}
+
+TSharedRef<SWidget> SCurveTreeContent::MakeTreeContent(const FArguments& InArgs)
+{
+	const TSharedRef<SCurveEditorTree> NewCurveEditorTree = SNew(SCurveEditorTree, InArgs._CurveEditor);
+	CurveEditorTreeView = NewCurveEditorTree;
+	return SNew(SScrollBorder, NewCurveEditorTree)
+		[
+			NewCurveEditorTree
+		];
+}
+
+TSharedRef<SWidget> SCurveTreeContent::MakeFooter(const FArguments& InArgs)
+{
+	const TSharedRef<FSequencer> Sequencer = InArgs._Sequencer.ToSharedRef();
+	
+	PlayTimeDisplay = Sequencer->MakePlayTimeDisplay();
+	return SNew(SBorder)
+		.BorderImage(FAppStyle::GetBrush("ToolPanel.GroupBorder"))
+		.Clipping(EWidgetClipping::ClipToBounds)
+		.Visibility_Static(&GetPopoutTransportControlsVisibility_CurveEditor, StaticCastWeakPtr<ISequencer>(Sequencer.ToWeakPtr()))
+		[
+			SNew(SHorizontalBox)
+
+			+ SHorizontalBox::Slot()
+			.HAlign(HAlign_Fill)
+			[
+				MakePopoutTransportControls(Sequencer, InArgs._TabManager)
+			]
+
+			+ SHorizontalBox::Slot()
+			.AutoWidth()
+			.VAlign(VAlign_Center)
+			.HAlign(HAlign_Right)
+			[
+				SNew(SButton)
+				.VAlign(EVerticalAlignment::VAlign_Center)
+				.ButtonStyle(FAppStyle::Get(), "NoBorder")
+				.ContentPadding(FMargin(1, 0))
+				[
+					SNew(SHorizontalBox)
+
+					+ SHorizontalBox::Slot()
+					.AutoWidth()
+					.VAlign(VAlign_Center)
+					.HAlign(HAlign_Right)
+					.Padding(FMargin(3.f, 0.f, 0.f, 0.f))
+					[
+						SNew(SBorder)
+						.BorderImage(nullptr)
+						[
+							PlayTimeDisplay.ToSharedRef()
+						]
+					]
+				]
+			]
+		];
+}
+
+void SCurveTreeContent::ScrollIntoView(const TWeakViewModelPtr<IOutlinerExtension>& WeakViewModel)
+{
+	const FCurveModelSyncer* CurveModelSyncer = CurveModelSyncerAttr.Get();
+	if (!CurveModelSyncer)
+	{
+		return;
+	}
+	
+	const FCurveEditorTreeItemID Id = CurveModelSyncer->FindIdByViewModel(WeakViewModel);
+	if (Id.IsValid())
+	{
+		CurveEditorTreeView->RequestScrollIntoView(Id);
+	}
+}
+
+TSet<TWeakViewModelPtr<IOutlinerExtension>> SCurveTreeContent::GetAllItems() const
+{
+	const FCurveModelSyncer* CurveModelSyncer = CurveModelSyncerAttr.Get();
+	if (!CurveModelSyncer)
+	{
+		return {};
+	}
+	
+	TSet<TWeakViewModelPtr<IOutlinerExtension>> AllItems;
+	for (const TPair<FCurveEditorTreeItemID, FCurveEditorTreeItem>& Item : CurveEditor->GetTree()->GetAllItems())
+	{
+		const TSharedPtr<FViewModel> ViewModel = CurveModelSyncer->FindViewModelById(Item.Key).Pin();
+		const TViewModelPtr<IOutlinerExtension> OutlinerNode = ViewModel ? CastViewModel<IOutlinerExtension>(ViewModel.ToSharedRef()) : nullptr;
+		if (OutlinerNode)
+		{
+			AllItems.Add(OutlinerNode);
+		}
+	}
+	return AllItems;
+}
+
+TSet<TWeakViewModelPtr<IOutlinerExtension>> SCurveTreeContent::GetSelectedItems() const
+{
+	const FCurveModelSyncer* CurveModelSyncer = CurveModelSyncerAttr.Get();
+	if (!CurveModelSyncer)
+	{
+		return {};
+	}
+	
+	TSet<TWeakViewModelPtr<IOutlinerExtension>> Selected;
+	for (const FCurveEditorTreeItemID& Item : CurveEditorTreeView->GetSelectedItems())
+	{
+		const TSharedPtr<FViewModel> ViewModel = CurveModelSyncer->FindViewModelById(Item).Pin();
+		const TViewModelPtr<IOutlinerExtension> OutlinerNode = ViewModel ? CastViewModel<IOutlinerExtension>(ViewModel.ToSharedRef()) : nullptr;
+		if (OutlinerNode)
+		{
+			Selected.Add(OutlinerNode);
+		}
+	}
+	return Selected;
+	
+}
+} // namespace UE::Sequencer
diff --git a/Engine/Source/Editor/Sequencer/Private/CurveEditor/Widgets/Tree/SCurveTreeContent.h b/Engine/Source/Editor/Sequencer/Private/CurveEditor/Widgets/Tree/SCurveTreeContent.h
new file mode 100644
--- /dev/null
+++ b/Engine/Source/Editor/Sequencer/Private/CurveEditor/Widgets/Tree/SCurveTreeContent.h
@@ -0,0 +1,84 @@
+﻿// Copyright Epic Games, Inc. All Rights Reserved.
+
+#pragma once
+
+#include "MVVM/Extensions/IOutlinerExtension.h"
+#include "Widgets/DeclarativeSyntaxSupport.h"
+#include "Widgets/SCompoundWidget.h"
+
+class FCurveEditor;
+class FSequencer;
+class FTabManager;
+class SCurveEditorTree;
+template<typename T> struct STemporarilyFocusedSpinBox;
+
+namespace UE::Sequencer
+{
+class FCurveModelSyncer;
+class FLinkedFilterViewModel;
+class SCurveEditorSearchAndFilterRow;
+
+/** Displays Curve Editor in Sequencer with linked filtering. */
+class SCurveTreeContent : public SCompoundWidget
+{
+public:
+	
+	SLATE_BEGIN_ARGS(SCurveTreeContent) {}
+		/** Required. The Sequencer instance that owns this content */
+		SLATE_ARGUMENT(TSharedPtr<FSequencer>, Sequencer)
+		/** Required. The curve editor instance that is being displayed */
+		SLATE_ARGUMENT(TSharedPtr<FCurveEditor>, CurveEditor)
+		/** Required. Determines the filtering widgets displayed. */
+		SLATE_ARGUMENT(TSharedPtr<FLinkedFilterViewModel>, FilterViewModel)
+		
+		/** Explicit tab manager to use, otherwise uses the global tab manager */
+		SLATE_ARGUMENT(TSharedPtr<FTabManager>, TabManager)
+
+		/** Used to translate between IOutlinerExtension and FCurveEditorTreeItemID. */
+		SLATE_ATTRIBUTE(const FCurveModelSyncer*, CurveModelSyncer)
+
+		/** Shows the filter pills. */
+		SLATE_NAMED_SLOT(FArguments, FilterPills)
+	SLATE_END_ARGS()
+	
+	void Construct(const FArguments& InArgs);
+	
+	/** Sets keyboard focus to the search box. */
+	void FocusSearchBox();
+	/** Sets keyboard focus to the play time display. */
+	void FocusPlayTimeDisplay();
+	
+	/** Syncs the tree view's selection to that of Sequencer. */
+	void SyncSelection() const;
+	
+	const TSharedPtr<SCurveEditorTree>& GetCurveEditorTreeView() const { return CurveEditorTreeView; }
+	
+private:
+	
+	/** The owning Sequencer instance that the Curve Editor is being displayed for. */
+	TWeakPtr<FSequencer> WeakSequencer;
+	/** The curve editor being displayed. */
+	TSharedPtr<FCurveEditor> CurveEditor;
+	
+	/** Used to translate between IOutlinerExtension and FCurveEditorTreeItemID. */
+	TAttribute<const FCurveModelSyncer*> CurveModelSyncerAttr;
+	
+	/** The filter row content displayed above the tree. */
+	TSharedPtr<SCurveEditorSearchAndFilterRow> FilterRow;
+	/** Curve editor tree widget */
+	TSharedPtr<SCurveEditorTree> CurveEditorTreeView;
+	/** The current playback time display. */
+	TSharedPtr<STemporarilyFocusedSpinBox<double>> PlayTimeDisplay;
+
+	TSharedRef<SWidget> MakeFilterArea(const FArguments& InArgs);
+	TSharedRef<SWidget> MakeTreeContent(const FArguments& InArgs);
+	TSharedRef<SWidget> MakeFooter(const FArguments& InArgs);
+	
+	/** Scrolls the given item into view. */
+	void ScrollIntoView(const TWeakViewModelPtr<IOutlinerExtension>& WeakViewModel);
+	/** @return All items visible in the tree view. */
+	TSet<TWeakViewModelPtr<IOutlinerExtension>> GetAllItems() const;
+	/** @return All items selected in the tree view. */
+	TSet<TWeakViewModelPtr<IOutlinerExtension>> GetSelectedItems() const;
+};
+} // namespace UE::Sequencer
diff --git a/Engine/Source/Editor/Sequencer/Private/CurveEditor/Widgets/Tree/SOldCurveTreeContent.cpp b/Engine/Source/Editor/Sequencer/Private/CurveEditor/Widgets/Tree/SOldCurveTreeContent.cpp
new file mode 100644
--- /dev/null
+++ b/Engine/Source/Editor/Sequencer/Private/CurveEditor/Widgets/Tree/SOldCurveTreeContent.cpp
@@ -0,0 +1,184 @@
+﻿// Copyright Epic Games, Inc. All Rights Reserved.
+
+#include "SOldCurveTreeContent.h"
+
+#include "CurveEditor.h"
+#include "CurveEditor/Views/SelectionSyncUtils.h"
+#include "LevelSequence.h"
+#include "Misc/TransportControlsHelper.h"
+#include "SPopoutTabInlineContent.h"
+#include "Sequencer.h"
+#include "SequencerSettings.h"
+#include "Styling/AppStyle.h"
+#include "Tree/SCurveEditorTree.h"
+#include "Tree/SCurveEditorTreeFilterStatusBar.h"
+#include "Tree/SCurveEditorTreeTextFilter.h"
+#include "Widgets/Docking/SDockTab.h"
+#include "Widgets/Layout/SBorder.h"
+#include "Widgets/Layout/SScrollBorder.h"
+#include "Widgets/SBoxPanel.h"
+#include "Widgets/SOverlay.h"
+#include "WorkspaceMenuStructureModule.h"
+
+#define LOCTEXT_NAMESPACE "SOldCurveTreeContent"
+
+namespace UE::Sequencer
+{
+void SOldCurveTreeContent::Construct(
+	const FArguments& InArgs, const TSharedRef<FSequencer> InSequencer, const TSharedRef<FCurveEditor>& InCurveEditor
+	)
+{
+	WeakSequencer = InSequencer;
+	CurveEditor = InCurveEditor;
+	
+	ChildSlot
+	[
+		SNew(SVerticalBox)
+
+		+ SVerticalBox::Slot()
+		.AutoHeight()
+		[
+			SNew(SBorder)
+			.BorderImage(FAppStyle::GetBrush("ToolPanel.GroupBorder"))
+			.Clipping(EWidgetClipping::ClipToBounds)
+			[
+				SAssignNew(CurveEditorSearchBox, SCurveEditorTreeTextFilter, InCurveEditor)
+			]
+		]
+
+		+ SVerticalBox::Slot()
+		[
+			MakeTreeContent(InCurveEditor)
+		]
+
+		+ SVerticalBox::Slot()
+		.AutoHeight()
+		[
+			MakeFooter(InArgs, InSequencer)
+		]
+	];
+}
+
+void SOldCurveTreeContent::FocusSearchBox()
+{
+	FSlateApplication::Get().SetKeyboardFocus(CurveEditorSearchBox, EFocusCause::SetDirectly);
+}
+
+void SOldCurveTreeContent::FocusPlayTimeDisplay()
+{
+	PlayTimeDisplay->Setup(); 
+	FSlateApplication::Get().SetKeyboardFocus(PlayTimeDisplay, EFocusCause::SetDirectly);
+}
+
+void SOldCurveTreeContent::SyncSelection()
+{
+	if (const TSharedPtr<FSequencer> Sequencer = WeakSequencer.Pin())
+	{
+		ApplySequencerToCurveEditorSelection(
+			Sequencer->GetViewModel()->GetSelection(), CurveEditor.ToSharedRef(), CurveEditorTreeView.ToSharedRef()
+			);
+		ExpandSelectedCurveEditorItems(*CurveEditor->GetTree(), *CurveEditorTreeView);
+	}
+}
+
+TSharedRef<SWidget> SOldCurveTreeContent::MakeTreeContent(const TSharedRef<FCurveEditor>& InCurveEditor)
+{
+	CurveEditorTreeView = SNew(SCurveEditorTree, InCurveEditor);
+	return SNew(SOverlay)
+
+		+ SOverlay::Slot()
+		[
+			SNew(SScrollBorder, CurveEditorTreeView.ToSharedRef())
+			[
+				CurveEditorTreeView.ToSharedRef()
+			]
+		]
+
+		+ SOverlay::Slot()
+		.VAlign(VAlign_Bottom)
+		[
+			SNew(SHorizontalBox)
+
+			+ SHorizontalBox::Slot()
+			.AutoWidth()
+			[
+				SAssignNew(CurveEditorTreeFilterStatusBar, SCurveEditorTreeFilterStatusBar, InCurveEditor)
+				.Visibility(EVisibility::Hidden) // Initially hidden, visible on hover of the info button
+			]
+		];
+}
+
+TSharedRef<SWidget> SOldCurveTreeContent::MakeFooter(const FArguments& InArgs, const TSharedRef<FSequencer> InSequencer)
+{
+	PlayTimeDisplay = InSequencer->MakePlayTimeDisplay();
+	
+	return SNew(SBorder)
+		.BorderImage(FAppStyle::GetBrush("ToolPanel.GroupBorder"))
+		.Clipping(EWidgetClipping::ClipToBounds)
+		.Visibility_Static(&UE::Sequencer::GetPopoutTransportControlsVisibility_CurveEditor, StaticCastWeakPtr<ISequencer>(InSequencer.ToWeakPtr()))
+		[
+			SNew(SHorizontalBox)
+
+			+ SHorizontalBox::Slot()
+			.AutoWidth()
+			.VAlign(VAlign_Center)
+			.HAlign(HAlign_Left)
+			[
+				SNew(SButton)
+				.VAlign(EVerticalAlignment::VAlign_Center)
+				.ButtonStyle(FAppStyle::Get(), "SimpleButton")
+				.ToolTipText_Lambda([this] { return LOCTEXT("ShowStatus", "Show Status"); })
+				.ContentPadding(FMargin(1, 0))
+				.OnHovered_Lambda([this] { CurveEditorTreeFilterStatusBar->ShowStatusBar(); })
+				.OnUnhovered_Lambda([this] { CurveEditorTreeFilterStatusBar->FadeOutStatusBar(); })
+				.OnClicked_Lambda([this] { CurveEditorTreeFilterStatusBar->HideStatusBar(); return FReply::Handled(); })
+				[
+					SNew(SImage)
+					.ColorAndOpacity(FSlateColor::UseForeground())
+					.Image(FAppStyle::Get().GetBrush("Icons.Info.Small"))
+				]
+			]
+
+			+ SHorizontalBox::Slot()
+			.HAlign(HAlign_Fill)
+			[
+				MakePopoutTransportControls(InSequencer, InArgs._TabManager)
+			]
+
+			+ SHorizontalBox::Slot()
+			.AutoWidth()
+			.VAlign(VAlign_Center)
+			.HAlign(HAlign_Right)
+			[
+				SNew(SButton)
+				.VAlign(EVerticalAlignment::VAlign_Center)
+				.ButtonStyle(FAppStyle::Get(), "NoBorder")
+				.ContentPadding(FMargin(1, 0))
+				[
+					SNew(SHorizontalBox)
+
+					+ SHorizontalBox::Slot()
+					.AutoWidth()
+					.VAlign(VAlign_Center)
+					.HAlign(HAlign_Right)
+					.Padding(FMargin(3.f, 0.f, 0.f, 0.f))
+					[
+						SNew(SBorder)
+						.BorderImage(nullptr)
+						[
+							PlayTimeDisplay.ToSharedRef()
+						]
+					]
+				]
+			]
+		];
+}
+
+USequencerSettings* SOldCurveTreeContent::GetSequencerSettings() const
+{
+	const TSharedPtr<ISequencer> Sequencer = WeakSequencer.Pin();
+	return Sequencer.IsValid() ? Sequencer->GetSequencerSettings() : nullptr;
+}
+} // namespace UE::Sequencer
+
+#undef LOCTEXT_NAMESPACE
\ No newline at end of file
diff --git a/Engine/Source/Editor/Sequencer/Private/CurveEditor/Widgets/Tree/SOldCurveTreeContent.h b/Engine/Source/Editor/Sequencer/Private/CurveEditor/Widgets/Tree/SOldCurveTreeContent.h
new file mode 100644
--- /dev/null
+++ b/Engine/Source/Editor/Sequencer/Private/CurveEditor/Widgets/Tree/SOldCurveTreeContent.h
@@ -0,0 +1,66 @@
+﻿// Copyright Epic Games, Inc. All Rights Reserved.
+
+#pragma once
+
+#include "STemporarilyFocusedSpinBox.h"
+#include "Widgets/DeclarativeSyntaxSupport.h"
+#include "Widgets/SCompoundWidget.h"
+
+class FCurveEditor;
+class FSequencer;
+class FTabManager;
+class SCurveEditorTree;
+class SCurveEditorTreeFilterStatusBar;
+class USequencerSettings;
+
+namespace UE::Sequencer
+{
+/**
+ * Curve Editor tree content before 5.8. 
+ * 
+ * This widget will be removed in the future. 
+ * For now, it is used to get the legacy UI looks (see UE-363527)
+ */
+class SOldCurveTreeContent : public SCompoundWidget
+{
+public:
+	
+	SLATE_BEGIN_ARGS(SOldCurveTreeContent){}
+		/** Explicit tab manager to use, otherwise uses the global tab manager */
+		SLATE_ARGUMENT(TSharedPtr<FTabManager>, TabManager)
+	SLATE_END_ARGS()
+	
+	void Construct(const FArguments& InArgs, const TSharedRef<FSequencer> InSequencer, const TSharedRef<FCurveEditor>& InCurveEditor);
+	
+	/** Sets keyboard focus to the search box. */
+	void FocusSearchBox();
+	/** Sets keyboard focus to the play time display. */
+	void FocusPlayTimeDisplay();
+	
+	/** Syncs the tree view's selection to that of Sequencer. */
+	void SyncSelection();
+	
+	const TSharedPtr<SCurveEditorTree>& GetCurveEditorTreeView() const { return CurveEditorTreeView; }
+	
+private:
+	
+	/** The owning Sequencer instance that the Curve Editor is being displayed for. */
+	TWeakPtr<FSequencer> WeakSequencer;
+	/** The curve editor being displayed. */
+	TSharedPtr<FCurveEditor> CurveEditor;
+	
+	/** Curve editor tree widget */
+	TSharedPtr<SCurveEditorTree> CurveEditorTreeView;
+	/** Filter Status Bar */
+	TSharedPtr<SCurveEditorTreeFilterStatusBar> CurveEditorTreeFilterStatusBar;
+	/** The current playback time display. */
+	TSharedPtr<STemporarilyFocusedSpinBox<double>> PlayTimeDisplay;
+	/** The search widget for filtering curves in the Curve Editor tree. */
+	TSharedPtr<SWidget> CurveEditorSearchBox;
+	
+	TSharedRef<SWidget> MakeTreeContent(const TSharedRef<FCurveEditor>& InCurveEditor);
+	TSharedRef<SWidget> MakeFooter(const FArguments& InArgs, const TSharedRef<FSequencer> InSequencer);
+	
+	USequencerSettings* GetSequencerSettings() const;
+};
+} // namespace UE::Sequencer
diff --git a/Engine/Source/Editor/Sequencer/Private/MVVM/CurveEditorIntegrationExtension.cpp b/Engine/Source/Editor/Sequencer/Private/MVVM/CurveEditorIntegrationExtension.cpp
--- a/Engine/Source/Editor/Sequencer/Private/MVVM/CurveEditorIntegrationExtension.cpp
+++ b/Engine/Source/Editor/Sequencer/Private/MVVM/CurveEditorIntegrationExtension.cpp
@@ -3,19 +3,13 @@
 #include "MVVM/CurveEditorIntegrationExtension.h"
 
 #include "CurveEditor.h"
-#include "Filters/ViewModels/Linking/LinkedFilterViewModel.h"
+#include "CurveEditor/SequencerCurveEditorApp.h"
 #include "MVVM/CurveEditorExtension.h"
-#include "MVVM/Extensions/ICurveEditorTreeItemExtension.h"
-#include "MVVM/SharedViewModelData.h"
-#include "MVVM/TrackModelStorageExtension.h"
 #include "MVVM/ViewModelPtr.h"
 #include "MVVM/ViewModels/SequenceModel.h"
 #include "MVVM/ViewModels/SequencerEditorViewModel.h"
 #include "MVVM/ViewModels/ViewModel.h"
-#include "MVVM/ViewModels/ViewModelIterators.h"
 #include "SequencerSettings.h"
-#include "CurveEditor/SequencerCurveEditorApp.h"
-#include "Views/CurveEditor/CurveEditorWidgetOwner.h"
 
 namespace UE
 {
diff --git a/Engine/Source/Editor/Sequencer/Private/MVVM/Views/CurveEditor/CurveEditorWidgetOwner.cpp b/Engine/Source/Editor/Sequencer/Private/MVVM/Views/CurveEditor/CurveEditorWidgetOwner.cpp
deleted file mode 100644
diff --git a/Engine/Source/Editor/Sequencer/Private/MVVM/Views/CurveEditor/CurveEditorWidgetOwner.h b/Engine/Source/Editor/Sequencer/Private/MVVM/Views/CurveEditor/CurveEditorWidgetOwner.h
deleted file mode 100644
diff --git a/Engine/Source/Editor/Sequencer/Private/MVVM/Views/CurveEditor/CurveEditorWidgets.h b/Engine/Source/Editor/Sequencer/Private/MVVM/Views/CurveEditor/CurveEditorWidgets.h
deleted file mode 100644
diff --git a/Engine/Source/Editor/Sequencer/Private/MVVM/Views/CurveEditor/SCurveEditorSearchAndFilterRow.cpp b/Engine/Source/Editor/Sequencer/Private/MVVM/Views/CurveEditor/SCurveEditorSearchAndFilterRow.cpp
deleted file mode 100644
diff --git a/Engine/Source/Editor/Sequencer/Private/MVVM/Views/CurveEditor/SCurveEditorSearchAndFilterRow.h b/Engine/Source/Editor/Sequencer/Private/MVVM/Views/CurveEditor/SCurveEditorSearchAndFilterRow.h
deleted file mode 100644
diff --git a/Engine/Source/Editor/Sequencer/Private/MVVM/Views/CurveEditor/SCurveTreeContent.cpp b/Engine/Source/Editor/Sequencer/Private/MVVM/Views/CurveEditor/SCurveTreeContent.cpp
deleted file mode 100644
diff --git a/Engine/Source/Editor/Sequencer/Private/MVVM/Views/CurveEditor/SCurveTreeContent.h b/Engine/Source/Editor/Sequencer/Private/MVVM/Views/CurveEditor/SCurveTreeContent.h
deleted file mode 100644
diff --git a/Engine/Source/Editor/Sequencer/Private/MVVM/Views/CurveEditor/SOldCurveTreeContent.cpp b/Engine/Source/Editor/Sequencer/Private/MVVM/Views/CurveEditor/SOldCurveTreeContent.cpp
deleted file mode 100644
diff --git a/Engine/Source/Editor/Sequencer/Private/MVVM/Views/CurveEditor/SOldCurveTreeContent.h b/Engine/Source/Editor/Sequencer/Private/MVVM/Views/CurveEditor/SOldCurveTreeContent.h
deleted file mode 100644
diff --git a/Engine/Source/Editor/Sequencer/Private/MVVM/Views/CurveEditor/SelectionSyncUtils.cpp b/Engine/Source/Editor/Sequencer/Private/MVVM/Views/CurveEditor/SelectionSyncUtils.cpp
deleted file mode 100644
diff --git a/Engine/Source/Editor/Sequencer/Private/MVVM/Views/CurveEditor/SelectionSyncUtils.h b/Engine/Source/Editor/Sequencer/Private/MVVM/Views/CurveEditor/SelectionSyncUtils.h
deleted file mode 100644
diff --git a/Engine/Source/Editor/Sequencer/Private/Widgets/CurveEditor/SSequencerCurveEditor.cpp b/Engine/Source/Editor/Sequencer/Private/Widgets/CurveEditor/SSequencerCurveEditor.cpp
deleted file mode 100644
diff --git a/Engine/Source/Editor/Sequencer/Private/Widgets/CurveEditor/SSequencerCurveEditor.h b/Engine/Source/Editor/Sequencer/Private/Widgets/CurveEditor/SSequencerCurveEditor.h
deleted file mode 100644
diff --git a/Engine/Source/Editor/Sequencer/Private/Widgets/CurveEditor/SequencerCurveEditorTimeSliderController.cpp b/Engine/Source/Editor/Sequencer/Private/Widgets/CurveEditor/SequencerCurveEditorTimeSliderController.cpp
deleted file mode 100644
diff --git a/Engine/Source/Editor/Sequencer/Private/Widgets/CurveEditor/SequencerCurveEditorTimeSliderController.h b/Engine/Source/Editor/Sequencer/Private/Widgets/CurveEditor/SequencerCurveEditorTimeSliderController.h
deleted file mode 100644
