Sequencer: Delete unused code (missed in last change)

#Synced-CL 51336034

diff --git a/Engine/Source/Editor/Sequencer/Public/MVVM/ViewModels/SectionOutlinerModel.h b/Engine/Source/Editor/Sequencer/Public/MVVM/ViewModels/SectionOutlinerModel.h
--- a/Engine/Source/Editor/Sequencer/Public/MVVM/ViewModels/SectionOutlinerModel.h
+++ b/Engine/Source/Editor/Sequencer/Public/MVVM/ViewModels/SectionOutlinerModel.h
@@ -132,7 +132,6 @@ public:
 	UE_API void OnDecorationRemoved(UObject* RemovedDecoration) override;
 
 	// Refreshes decoration models under this section outliner.
-	// Called during layout when bDecorationsNeedRefresh is true.
 	UE_API void RefreshDecorations();
 
 private:
@@ -145,7 +144,6 @@ private:
 	TSharedPtr<ISequencerTrackEditor> TrackEditor;
 
 	UE::MovieScene::TNonIntrusiveEventHandler<UE::MovieScene::ISectionEventHandler> SectionEventHandler;
-	bool bDecorationsNeedRefresh = false;
 };
 
 } // namespace Sequencer
