Sequencer: Fix decoration property display- we were double displaying section properties inside decorations. ISectionProviderExtension provides this now.

#Synced-CL 51336034

diff --git a/Engine/Source/Editor/Sequencer/Private/SequencerCommonHelpers.cpp b/Engine/Source/Editor/Sequencer/Private/SequencerCommonHelpers.cpp
--- a/Engine/Source/Editor/Sequencer/Private/SequencerCommonHelpers.cpp
+++ b/Engine/Source/Editor/Sequencer/Private/SequencerCommonHelpers.cpp
@@ -689,29 +689,10 @@ void SequencerHelpers::BuildEditDecorationMenu(const TWeakPtr<FSequencer>& InWea
 		UClass* DecorationClass = Pair.Key;
 		const TArray<TWeakObjectPtr<>>& ClassDecorations = Pair.Value;
 
-		// Check if this decoration class provides sections
-		// If so, we display the section properties under the decoration's class header
+		// Check if this class has any editable properties - skip if empty
 		TArray<TWeakObjectPtr<>> ObjectsToDisplay;
 		bool bIsSectionProvider = false;
-
-		for (const TWeakObjectPtr<>& WeakDecoration : ClassDecorations)
-		{
-			if (IMovieSceneSectionProviderDecoration* SectionProvider = Cast<IMovieSceneSectionProviderDecoration>(WeakDecoration.Get()))
-			{
-				bIsSectionProvider = true;
-				for (UMovieSceneSection* Section : SectionProvider->GetSections())
-				{
-					if (Section)
-					{
-						ObjectsToDisplay.AddUnique(Section);
-					}
-				}
-			}
-		}
-
-		if (!bIsSectionProvider)
 		{
-			// Check if this class has any editable properties - skip if empty
 			bool bHasEditableProperties = false;
 			for (TFieldIterator<FProperty> PropIt(DecorationClass); PropIt; ++PropIt)
 			{
@@ -726,8 +707,16 @@ void SequencerHelpers::BuildEditDecorationMenu(const TWeakPtr<FSequencer>& InWea
 				continue;
 			}
 
-			// Use the decorations themselves for display
 			ObjectsToDisplay = ClassDecorations;
+
+			for (const TWeakObjectPtr<>& WeakDecoration : ClassDecorations)
+			{
+				if (Cast<IMovieSceneSectionProviderDecoration>(WeakDecoration.Get()))
+				{
+					bIsSectionProvider = true;
+					break;
+				}
+			}
 		}
 
 		if (ObjectsToDisplay.IsEmpty())
