Sequencer: Visbility of filter pills is now saved per UI area instead of per filter bar instance.

- New USequencerSettings::FindOrAddTrackFilterArea for storing settings that are per filter area, regardless whether in linked or instanced mode. Visibility is the first such property.
- New FFilterAreaClient: Responsible for changing the visbility of filter pills in any UI area.
- FLinkedFilteringMediator now uses FFilterAreaClient to affect the visibility of the filter pills in Sequencer's Outliner.
- SFilterBarSwitcher handles updating the widget visbility depending on the new setting
- Deprecated ISequencerFilterBar functions related to filter bar visibility, as it's now per UI area instead of per filter instance.

Curve Editor will have leverage FFilterAreaClient to affect its filter bar visibility in a follow up change.

#Synced-CL 51334859

diff --git a/Engine/Plugins/Experimental/SequenceNavigator/Source/SequenceNavigator/Public/Filters/NavigationToolFilterBar.h b/Engine/Plugins/Experimental/SequenceNavigator/Source/SequenceNavigator/Public/Filters/NavigationToolFilterBar.h
--- a/Engine/Plugins/Experimental/SequenceNavigator/Source/SequenceNavigator/Public/Filters/NavigationToolFilterBar.h
+++ b/Engine/Plugins/Experimental/SequenceNavigator/Source/SequenceNavigator/Public/Filters/NavigationToolFilterBar.h
@@ -44,6 +44,7 @@ public:
 
 	DECLARE_EVENT_TwoParams(ISequencerFilterBar, FNavigationToolFiltersChanged, const ENavigationToolFilterChange /*InChangeType*/, const TSharedRef<FNavigationToolFilter>& /*InFilter*/);
 	DECLARE_EVENT_OneParam(ISequencerFilterBar, FOnLayoutChanged, EFilterBarLayout /** InNewLayout */);
+	DECLARE_MULTICAST_DELEGATE_OneParam(FOnVisibilityChanged, bool /*bIsVisible*/)
 	
 	static FCustomTextFilterData DefaultNewCustomTextFilterData(const FText& InFilterString);
 
@@ -70,7 +71,7 @@ public:
 	FNavigationToolFiltersChanged& OnFiltersChanged() { return FiltersChangedEvent; }
 	FOnLayoutChanged& OnLayoutChanged() { return LayoutChangedEvent; }
 	FSimpleMulticastDelegate& OnActivationStateChanged() { return ActivationStateChangedEvent; }
-	virtual FOnVisibilityChanged& OnVisibilityChanged() override { return VisibilityChangedEvent; }
+	FOnVisibilityChanged& OnVisibilityChanged() { return VisibilityChangedEvent; }
 	virtual FSimpleMulticastDelegate& OnRequestUpdate() override { return RequestUpdateEvent; }
 	virtual FSimpleMulticastDelegate& OnTextFilterTextChanged() override { return TextFilterTextChangedEvent; }
 	virtual FMuteStateChangedDelegate& OnMuteFiltersChanged() override { return MuteFiltersChangedEvent; }
@@ -191,10 +192,9 @@ public:
 
 	TSharedRef<SComboButton> MakeAddFilterButton();
 
-	virtual bool ShouldShowFilterBarWidget() const override;
-
-	virtual bool IsFilterBarVisible() const override;
-	virtual void ToggleFilterBarVisibility() override;
+	bool ShouldShowFilterBarWidget() const;
+	bool IsFilterBarVisible() const;
+	void ToggleFilterBarVisibility();
 
 	bool IsFilterBarLayout(const EFilterBarLayout InLayout) const;
 	void SetToVerticalLayout();
diff --git a/Engine/Source/Editor/Sequencer/Private/CurveEditor/Views/CurveEditorWidgetOwner.cpp b/Engine/Source/Editor/Sequencer/Private/CurveEditor/Views/CurveEditorWidgetOwner.cpp
--- a/Engine/Source/Editor/Sequencer/Private/CurveEditor/Views/CurveEditorWidgetOwner.cpp
+++ b/Engine/Source/Editor/Sequencer/Private/CurveEditor/Views/CurveEditorWidgetOwner.cpp
@@ -17,6 +17,7 @@
 #include "SCurveKeyDetailPanel.h"
 #include "Sequencer.h"
 #include "SequencerCommands.h"
+#include "Filters/FilterConfigIdentifiers.h"
 #include "Widgets/DeclarativeSyntaxSupport.h"
 #include "Widgets/Docking/SDockTab.h"
 
@@ -52,6 +53,9 @@ static FCurveEditorWidgets MakePresenterWidgets(
 		.FilterPills()
 		[
 			SNew(SFilterBarSwitcher)
+			.Sequencer(InSequencer)
+			// TODO UE-368977: Introduce separate setting for curve editor filter bar visibility. For now, use the same as in Sequencer Outliner.
+			.FilterAreaConfigId(ConfigIds::FilterArea_Outliner)
 			.LinkedFilterViewModel(InFilteringViewModel)
 		];
 	
diff --git a/Engine/Source/Editor/Sequencer/Private/Filters/FilterConfigIdentifiers.cpp b/Engine/Source/Editor/Sequencer/Private/Filters/FilterConfigIdentifiers.cpp
--- a/Engine/Source/Editor/Sequencer/Private/Filters/FilterConfigIdentifiers.cpp
+++ b/Engine/Source/Editor/Sequencer/Private/Filters/FilterConfigIdentifiers.cpp
@@ -8,4 +8,7 @@ const TCHAR* LinkedConfigSubId = TEXT(""); // TODO UE-362151: For now empty so w
 
 const TCHAR* Sequencer_InstancedConfigSubId = TEXT("Outliner");
 const TCHAR* CurveEditor_InstancedConfigId = TEXT("CurveEditor");
+
+const TCHAR* FilterArea_Outliner = TEXT("Outliner");
+
 } // namespace UE::Sequencer
\ No newline at end of file
diff --git a/Engine/Source/Editor/Sequencer/Private/Filters/FilterConfigIdentifiers.h b/Engine/Source/Editor/Sequencer/Private/Filters/FilterConfigIdentifiers.h
--- a/Engine/Source/Editor/Sequencer/Private/Filters/FilterConfigIdentifiers.h
+++ b/Engine/Source/Editor/Sequencer/Private/Filters/FilterConfigIdentifiers.h
@@ -26,4 +26,10 @@ extern const TCHAR* Sequencer_InstancedConfigSubId;
  * @see FSequencerFilterBar::GetIdentifier.
  */
 extern const TCHAR* CurveEditor_InstancedConfigId;
+
+/** 
+ * ID to use for the Outliner filter area.
+ * @see USequencerSettings::FindOrAddTrackFilterArea 
+ */
+extern const TCHAR* FilterArea_Outliner;
 } // namespace UE::Sequencer
diff --git a/Engine/Source/Editor/Sequencer/Private/Filters/Menus/SequencerViewOptionsMenu.cpp b/Engine/Source/Editor/Sequencer/Private/Filters/Menus/SequencerViewOptionsMenu.cpp
--- a/Engine/Source/Editor/Sequencer/Private/Filters/Menus/SequencerViewOptionsMenu.cpp
+++ b/Engine/Source/Editor/Sequencer/Private/Filters/Menus/SequencerViewOptionsMenu.cpp
@@ -37,7 +37,7 @@ TSharedRef<SWidget> FSequencerViewOptionsMenu::CreateMenu(const TWeakPtr<FSequen
 	ContextObject->Init(InSequencerWeak);
 	ContextObject->OnPopulateFilterBarMenu = FOnPopulateFilterBarMenu::CreateSP(this, &FSequencerViewOptionsMenu::PopulateMenu);
 
-	const FToolMenuContext MenuContext(InSequencerWeak.Pin()->GetFilterInterface()->GetCommandList(), nullptr, ContextObject);
+	const FToolMenuContext MenuContext(InSequencerWeak.Pin()->GetCommandBindings(), nullptr, ContextObject);
 	return UToolMenus::Get()->GenerateWidget(FilterMenuName, MenuContext);
 }
 
diff --git a/Engine/Source/Editor/Sequencer/Private/Filters/SequencerFilterAreaConfig.cpp b/Engine/Source/Editor/Sequencer/Private/Filters/SequencerFilterAreaConfig.cpp
new file mode 100644
--- /dev/null
+++ b/Engine/Source/Editor/Sequencer/Private/Filters/SequencerFilterAreaConfig.cpp
@@ -0,0 +1,13 @@
+﻿// Copyright Epic Games, Inc. All Rights Reserved.
+
+#include "Filters/SequencerFilterAreaConfig.h"
+
+void FSequencerFilterAreaConfig::SetIsFilterBarVisible(bool bNewIsVisible)
+{
+	bIsFilterBarVisible = bNewIsVisible;
+}
+
+bool FSequencerFilterAreaConfig::IsFilterBarVisible() const
+{
+	return bIsFilterBarVisible;
+}
diff --git a/Engine/Source/Editor/Sequencer/Private/Filters/SequencerFilterBar.cpp b/Engine/Source/Editor/Sequencer/Private/Filters/SequencerFilterBar.cpp
--- a/Engine/Source/Editor/Sequencer/Private/Filters/SequencerFilterBar.cpp
+++ b/Engine/Source/Editor/Sequencer/Private/Filters/SequencerFilterBar.cpp
@@ -28,6 +28,7 @@
 #include "MVVM/Views/SOutlinerView.h"
 #include "Sequencer.h"
 #include "SSequencer.h"
+#include "Filters/SequencerFilterAreaConfig.h"
 #include "Stats/StatsMisc.h"
 #include "UObject/UObjectIterator.h"
 #include "Utils/HideIsolateViewModelUtils.h"
@@ -36,8 +37,6 @@ using namespace UE::Sequencer;
 
 #define LOCTEXT_NAMESPACE "SequencerFilterBar"
 
-const FName FSequencerFilterBar::SharedIdentifier = TEXT("SharedSequencerTrackFilter");
-
 int32 FSequencerFilterBar::InstanceCount = 0;
 
 FSequencerFilterBar::FSequencerFilterBar(FSequencer& InSequencer, FString InSubIdentifier)
@@ -204,11 +203,6 @@ void FSequencerFilterBar::CreateDefaultFilters()
 void FSequencerFilterBar::BindCommands()
 {
 	const FSequencerTrackFilterCommands& TrackFilterCommands = FSequencerTrackFilterCommands::Get();
-	
-	CommandList->MapAction(TrackFilterCommands.ToggleFilterBarVisibility,
-		FExecuteAction::CreateSP(this, &FSequencerFilterBar::ToggleFilterBarVisibility),
-		FCanExecuteAction(),
-		FIsActionChecked::CreateSP(this, &FSequencerFilterBar::IsFilterBarVisible));
 
 	CommandList->MapAction(TrackFilterCommands.ResetFilters,
 		FExecuteAction::CreateSP(this, &FSequencerFilterBar::ResetFilters),
@@ -1622,37 +1616,6 @@ bool FSequencerFilterBar::ShouldUpdateOnTrackValueChanged() const
 	return false;
 }
 
-bool FSequencerFilterBar::ShouldShowFilterBarWidget() const
-{
-	if (const USequencerSettings* const Settings = Sequencer.GetSequencerSettings())
-	{
-		return Settings->IsFilterBarVisible() && HasAnyFiltersEnabled();
-	}
-	return false;
-}
-
-bool FSequencerFilterBar::IsFilterBarVisible() const
-{
-	if (const USequencerSettings* const Settings = Sequencer.GetSequencerSettings())
-	{
-		return Settings->IsFilterBarVisible();
-	}
-	return false;
-}
-
-void FSequencerFilterBar::ToggleFilterBarVisibility()
-{
-	USequencerSettings* const Settings = Sequencer.GetSequencerSettings();
-	if (!Settings)
-	{
-		return;
-	}
-
-	const bool bNewVisibility = !Settings->IsFilterBarVisible();
-	Settings->SetFilterBarVisible(bNewVisibility);
-	VisibilityChangedEvent.Broadcast(bNewVisibility);
-}
-
 bool FSequencerFilterBar::IsFilterSupported(const TSharedRef<FSequencerTrackFilter>& InFilter) const
 {
 	UMovieSceneSequence* const MovieSceneSequence = Sequencer.GetFocusedMovieSceneSequence();
diff --git a/Engine/Source/Editor/Sequencer/Private/Filters/SequencerFilterBar.h b/Engine/Source/Editor/Sequencer/Private/Filters/SequencerFilterBar.h
--- a/Engine/Source/Editor/Sequencer/Private/Filters/SequencerFilterBar.h
+++ b/Engine/Source/Editor/Sequencer/Private/Filters/SequencerFilterBar.h
@@ -43,21 +43,17 @@ enum class ESequencerFilterChange : uint8
 class FSequencerFilterBar : public ISequencerTrackFilters
 {
 public:
-	/** An identifier shared by all filter bars, used to save and load settings common to every instance */
-	static const FName SharedIdentifier;
 	
 	DECLARE_EVENT_TwoParams(FSequencerFilterBar, FSequencerFiltersChanged, const ESequencerFilterChange /*InChangeType*/, const TSharedRef<FSequencerTrackFilter>& /*InFilter*/);
 	DECLARE_EVENT_TwoParams(FSequencerFilterBar, FSequencerCustomTextFiltersChanged, const ESequencerFilterChange /*InChangeType*/, const TSharedRef<FSequencerTrackFilter_CustomText>& /*InFilter*/);
 	
-	
-
-	/**
-	 * @param InSequencer The owning Sequencer instance 
-	 * @param InSubIdentifier Used to identify config entries. This is relevant saving settings to config entries corresponding to linked & instanced
-	 *	filtering. The config entry has the format "{0}.{1}" where {0} = FSequencer::Settings->GetName(), and {1} is SubIdentifier. If empty, only the
-	 *	settings name is used, i.e. FSequencer::Settings->GetName().
+	/** 
+	 * @param InSequencer The owning Sequencer instance.
+	 * @param InSubIdentifier Sub-identifier that affects the result of FSequencerFilterBar::GetIdentifier. 
+	 *		This is relevant for saving per-filter bar instance config.
+	 * @see FSequencerFilterBarConfig, USequencerSettings::FindTrackFilterBar
 	 */
-	explicit FSequencerFilterBar(FSequencer& InSequencer, FString InSubIdentifier = TEXT(""));
+	explicit FSequencerFilterBar(FSequencer& InSequencer, FString InSubIdentifier);
 	virtual ~FSequencerFilterBar() override;
 
 	void BindCommands();
@@ -229,15 +225,9 @@ public:
 
 	bool ShouldUpdateOnTrackValueChanged() const;
 
-	virtual bool ShouldShowFilterBarWidget() const override;
-
-	virtual bool IsFilterBarVisible() const override;
-	virtual void ToggleFilterBarVisibility() override;
-
 	bool IsFilterSupported(const TSharedRef<FSequencerTrackFilter>& InFilter) const;
 	bool IsFilterSupported(const FString& InFilterName) const;
 
-	virtual FOnVisibilityChanged& OnVisibilityChanged() override { return VisibilityChangedEvent; }
 	virtual FSimpleMulticastDelegate& OnRequestUpdate() override { return RequestUpdateEvent; }
 	virtual FSimpleMulticastDelegate& OnTextFilterTextChanged() override { return OnTextFilterTextChangedEvent; }
 	virtual FMuteStateChangedDelegate& OnMuteFiltersChanged() override { return OnMuteFiltersChangedEvent; }
@@ -297,7 +287,6 @@ protected:
 	FSequencerFilterData FilterData;
 
 	FSequencerFiltersChanged FiltersChangedEvent;
-	FOnVisibilityChanged VisibilityChangedEvent;
 	FSimpleMulticastDelegate RequestUpdateEvent;
 	FSimpleMulticastDelegate OnTextFilterTextChangedEvent;
 	FMuteStateChangedDelegate OnMuteFiltersChangedEvent;
@@ -309,11 +298,9 @@ private:
 	static int32 InstanceCount;
 	
 	/**
-	 * Used to build the config entry for the filters.
+	 * Used for per-filter bar instance config.
 	 * This is relevant for the settings to use for reading and saving config.
-	 * 
-	 * The final key is "{0}.{1}" where {0} = FSequencer::Settings->GetFName(), and {1} is SubIdentifier.
-	 * Examples: "LevelSequenceEditor.Linked", "LevelSequenceEditor.Instanced".
+	 * @see FSequencerFilterBarConfig, USequencerSettings::FindTrackFilterBar, FSequencerFilterBar::GetIdentifier.
 	 */
 	const FString SubIdentifier;
 	
diff --git a/Engine/Source/Editor/Sequencer/Private/Filters/ViewModels/Linking/FilterAreaManager.cpp b/Engine/Source/Editor/Sequencer/Private/Filters/ViewModels/Linking/FilterAreaManager.cpp
new file mode 100644
--- /dev/null
+++ b/Engine/Source/Editor/Sequencer/Private/Filters/ViewModels/Linking/FilterAreaManager.cpp
@@ -0,0 +1,54 @@
+﻿// Copyright Epic Games, Inc. All Rights Reserved.
+
+#include "FilterAreaManager.h"
+
+#include "Sequencer.h"
+#include "Filters/SequencerTrackFilterCommands.h"
+
+namespace UE::Sequencer
+{
+FFilterAreaManager::FFilterAreaManager(
+	FName InFilterAreaConfigId,
+	TWeakPtr<FSequencer> InWeakSequencer, 
+	const TSharedRef<FLinkedFilterViewModel>& InFilterModel
+	)
+	: FilterAreaConfigId(InFilterAreaConfigId)
+	, WeakSequencer(MoveTemp(InWeakSequencer))
+	, FilterModel(InFilterModel)
+{
+	BindCommands();
+}
+
+void FFilterAreaManager::SetFilterBarVisible(bool bIsVisible)
+{
+	const TSharedPtr<FSequencer> SequencerPin = WeakSequencer.Pin();
+	USequencerSettings* Settings = SequencerPin ? SequencerPin->GetSequencerSettings() : nullptr;
+	if (Settings)
+	{
+		constexpr bool bSaveConfig = false; // Saving unnecessary because we'll save right after settings
+		Settings->FindOrAddTrackFilterArea(FilterAreaConfigId, bSaveConfig).SetIsFilterBarVisible(bIsVisible);
+		Settings->SaveConfig();
+	}
+}
+
+bool FFilterAreaManager::IsFilterBarVisible() const
+{
+	const TSharedPtr<FSequencer> SequencerPin = WeakSequencer.Pin();
+	USequencerSettings* Settings = SequencerPin ? SequencerPin->GetSequencerSettings() : nullptr;
+	
+	constexpr bool bSaveConfig = false; // Saving unnecessary because we're just getting info and using the default if not found.
+	return Settings && Settings->FindOrAddTrackFilterArea(FilterAreaConfigId, bSaveConfig).IsFilterBarVisible();
+}
+
+void FFilterAreaManager::BindCommands()
+{
+	FSequencerTrackFilterCommands& FilterCommands = FSequencerTrackFilterCommands::Get();
+	
+	CommandList->MapAction(
+		FilterCommands.ToggleFilterBarVisibility,
+		FExecuteAction::CreateRaw(this, &FFilterAreaManager::ToggleFilterBarVisibility),
+		FCanExecuteAction::CreateLambda([]{ return true; }),
+		FIsActionChecked::CreateRaw(this, &FFilterAreaManager::IsFilterBarVisible)
+		);
+}
+}
diff --git a/Engine/Source/Editor/Sequencer/Private/Filters/ViewModels/Linking/FilterAreaManager.h b/Engine/Source/Editor/Sequencer/Private/Filters/ViewModels/Linking/FilterAreaManager.h
new file mode 100644
--- /dev/null
+++ b/Engine/Source/Editor/Sequencer/Private/Filters/ViewModels/Linking/FilterAreaManager.h
@@ -0,0 +1,57 @@
+﻿// Copyright Epic Games, Inc. All Rights Reserved.
+
+#pragma once
+
+#include "Framework/Commands/UICommandList.h"
+#include "Templates/SharedPointer.h"
+#include "Templates/UnrealTemplate.h"
+
+class FSequencer;
+class FUICommandList;
+
+namespace UE::Sequencer
+{
+class FLinkedFilterViewModel;
+
+/** 
+ * Manages a UI area that displays filters that can be linked or instanced.
+ * It manages visibility of the filter pills and handles the related commands.
+ */
+class FFilterAreaManager : public FNoncopyable
+{
+public:
+	
+	explicit FFilterAreaManager(
+		FName InFilterAreaConfigId,
+		TWeakPtr<FSequencer> InWeakSequencer, 
+		const TSharedRef<FLinkedFilterViewModel>& InFilterModel
+		);
+	
+	/** Sets whether the filter bar pills should be visible. */
+	void SetFilterBarVisible(bool bIsVisible);
+	/** @return Whether the filter bar pills should be visible */
+	bool IsFilterBarVisible() const;
+	/** Inverts the filter bar pill visibility. */
+	void ToggleFilterBarVisibility() { SetFilterBarVisible(!IsFilterBarVisible()); }
+	
+	const TSharedRef<FLinkedFilterViewModel>& GetFilterModel() const { return FilterModel; }
+	const TSharedRef<FUICommandList>& GetFilterAreaCommandList() const { return CommandList; }
+	
+private:
+	
+	/** @see USequencerSettings::FindOrAddTrackFilterArea. */
+	const FName FilterAreaConfigId;
+	
+	/** Used to get the settings, e.g. to save filter bar visbility. */
+	const TWeakPtr<FSequencer> WeakSequencer;
+	
+	/** Manages the linked and unlinked filter state. */
+	const TSharedRef<FLinkedFilterViewModel> FilterModel;
+	
+	/** Handles the user's commands, e.g. to change visibility. */
+	const TSharedRef<FUICommandList> CommandList = MakeShared<FUICommandList>();
+	
+	void BindCommands();
+};
+}
+
diff --git a/Engine/Source/Editor/Sequencer/Private/Filters/Widgets/Linking/SFilterBarSwitcher.cpp b/Engine/Source/Editor/Sequencer/Private/Filters/Widgets/Linking/SFilterBarSwitcher.cpp
--- a/Engine/Source/Editor/Sequencer/Private/Filters/Widgets/Linking/SFilterBarSwitcher.cpp
+++ b/Engine/Source/Editor/Sequencer/Private/Filters/Widgets/Linking/SFilterBarSwitcher.cpp
@@ -2,6 +2,7 @@
 
 #include "SFilterBarSwitcher.h"
 
+#include "Sequencer.h"
 #include "SFilterLinkStateSwitcher.h"
 #include "Filters/Widgets/SSequencerFilterBar.h"
 #include "Filters/ViewModels/Linking/LinkedFilterViewModel.h"
@@ -10,6 +11,10 @@ namespace UE::Sequencer
 {
 void SFilterBarSwitcher::Construct(const FArguments& InArgs)
 {
+	WeakSequencer = InArgs._Sequencer;
+	FilterAreaConfigId = InArgs._FilterAreaConfigId;
+	check(WeakSequencer.IsValid());
+	
 	const TSharedPtr<FLinkedFilterViewModel> LinkedFilterViewModel = InArgs._LinkedFilterViewModel;
 	check(LinkedFilterViewModel);
 	
@@ -17,6 +22,7 @@ void SFilterBarSwitcher::Construct(const FArguments& InArgs)
 	[
 		SAssignNew(WidgetSwitcher, SFilterLinkStateSwitcher)
 		.LinkedFilterViewModel(LinkedFilterViewModel)
+		.Visibility(this, &SFilterBarSwitcher::GetFilterBarVisibility)
 		.LinkedContent()
 		[
 			SNew(SSequencerFilterBar, LinkedFilterViewModel->GetLinkedFilterBarImpl())
@@ -34,4 +40,12 @@ FReply SFilterBarSwitcher::OnMouseButtonUp(const FGeometry& MyGeometry, const FP
 {
 	return WidgetSwitcher->GetActiveWidget()->OnMouseButtonUp(MyGeometry, MouseEvent);
 }
+
+EVisibility SFilterBarSwitcher::GetFilterBarVisibility() const
+{
+	const TSharedPtr<FSequencer> Sequencer = WeakSequencer.Pin();
+	USequencerSettings* Settings = Sequencer ? Sequencer->GetSequencerSettings() : nullptr;
+	const FSequencerFilterAreaConfig* Config = Settings ? Settings->FindTrackFilterArea(FilterAreaConfigId) : nullptr;
+	return Config && Config->IsFilterBarVisible() ? EVisibility::Visible : EVisibility::Collapsed;
+}
 } // namespace UE::Sequencer
diff --git a/Engine/Source/Editor/Sequencer/Private/Filters/Widgets/Linking/SFilterBarSwitcher.h b/Engine/Source/Editor/Sequencer/Private/Filters/Widgets/Linking/SFilterBarSwitcher.h
--- a/Engine/Source/Editor/Sequencer/Private/Filters/Widgets/Linking/SFilterBarSwitcher.h
+++ b/Engine/Source/Editor/Sequencer/Private/Filters/Widgets/Linking/SFilterBarSwitcher.h
@@ -17,10 +17,18 @@ class SFilterBarSwitcher : public SCompoundWidget
 public:
 	
 	SLATE_BEGIN_ARGS(SFilterBarSwitcher) {}
+		/** 
+		 * Used to determine whether the filter bar is visible. 
+		 * @see USequencerSettings::FindOrAddTrackFilterArea 
+		 */
+		SLATE_ARGUMENT(FName, FilterAreaConfigId)
+		/** Used to get the USequencerSettings to determine whether the filter bar is visible. */
+		SLATE_ARGUMENT(TWeakPtr<FSequencer>, Sequencer)
+		
 		/** View-model determines the ELinkedFilterMode driving the displayed content. */
 		SLATE_ARGUMENT(TSharedPtr<FLinkedFilterViewModel>, LinkedFilterViewModel)
 	SLATE_END_ARGS()
-	
+
 	void Construct(const FArguments& InArgs);
 	
 	//~ Begin SWidget
@@ -29,7 +37,15 @@ public:
 	
 private:
 	
+	/** Used to get the USequencerSettings to determine whether the filter bar is visible. */
+	TWeakPtr<FSequencer> WeakSequencer;
+	/** @see USequencerSettings::FindOrAddTrackFilterArea  */
+	FName FilterAreaConfigId;
+	
 	/** Determines the filter bar that is displayed. */
 	TSharedPtr<SFilterLinkStateSwitcher> WidgetSwitcher;
+	
+	/** @return Checks the settings whether the filter bar should be visible. */
+	EVisibility GetFilterBarVisibility() const;
 };
 } // namespace UE::Sequencer
diff --git a/Engine/Source/Editor/Sequencer/Private/MVVM/ViewModels/LinkedFilteringMediator.cpp b/Engine/Source/Editor/Sequencer/Private/MVVM/ViewModels/LinkedFilteringMediator.cpp
--- a/Engine/Source/Editor/Sequencer/Private/MVVM/ViewModels/LinkedFilteringMediator.cpp
+++ b/Engine/Source/Editor/Sequencer/Private/MVVM/ViewModels/LinkedFilteringMediator.cpp
@@ -8,64 +8,74 @@
 #include "Filters/ViewModels/Linking/ILinkedFilterViewModel.h"
 #include "Filters/ViewModels/Linking/LinkedFilterViewModel.h"
 #include "SequencerNodeTree.h"
+#include "Filters/FilterConfigIdentifiers.h"
 
 namespace UE::Sequencer
 {
-FLinkedFilteringMediator::FLinkedFilteringMediator(const TSharedRef<FLinkedFilterViewModel>& InSequencerFilterViewModel)
-	: SequencerFilterViewModel(InSequencerFilterViewModel)
+FLinkedFilteringMediator::FLinkedFilteringMediator(
+	TWeakPtr<FSequencer> InSequencer, const TSharedRef<FLinkedFilterViewModel>& InSequencerFilterViewModel
+	)
+	: FFilterAreaManager(ConfigIds::FilterArea_Outliner, MoveTemp(InSequencer), InSequencerFilterViewModel)
 {
-	SequencerFilterViewModel->OnFilterModeChanged().AddRaw(this, &FLinkedFilteringMediator::OnFilterModeChanged);
+	GetFilterModel()->OnFilterModeChanged().AddRaw(this, &FLinkedFilteringMediator::OnFilterModeChanged);
 	
-	const TSharedRef<FSequencerFilterBar> SharedFilterBar = SequencerFilterViewModel->GetLinkedFilterBarImpl();
+	const TSharedRef<FSequencerFilterBar> SharedFilterBar = GetFilterModel()->GetLinkedFilterBarImpl();
 	SharedFilterBar->OnChangeOutlinerExtensionFilterState().AddRaw(this, &FLinkedFilteringMediator::OnFilterStateChanged_LinkedFilterBar);
 	
-	const TSharedRef<FSequencerFilterBar> InstancedFilterBar = SequencerFilterViewModel->GetInstancedFilterBarImpl();
+	const TSharedRef<FSequencerFilterBar> InstancedFilterBar = GetFilterModel()->GetInstancedFilterBarImpl();
 	InstancedFilterBar->OnChangeOutlinerExtensionFilterState().AddRaw(this, &FLinkedFilteringMediator::OnFilterStateChanged_InstancedFilterBar);
 }
 
 FLinkedFilteringMediator::~FLinkedFilteringMediator()
 {
-	SequencerFilterViewModel->OnFilterModeChanged().RemoveAll(this);
+	GetFilterModel()->OnFilterModeChanged().RemoveAll(this);
 	
-	const TSharedRef<FSequencerFilterBar> SharedFilterBar = SequencerFilterViewModel->GetLinkedFilterBarImpl();
+	const TSharedRef<FSequencerFilterBar> SharedFilterBar = GetFilterModel()->GetLinkedFilterBarImpl();
 	SharedFilterBar->OnChangeOutlinerExtensionFilterState().RemoveAll(this);
 	
-	const TSharedRef<FSequencerFilterBar> InstancedFilterBar = SequencerFilterViewModel->GetInstancedFilterBarImpl();
+	const TSharedRef<FSequencerFilterBar> InstancedFilterBar = GetFilterModel()->GetInstancedFilterBarImpl();
 	InstancedFilterBar->OnChangeOutlinerExtensionFilterState().RemoveAll(this);
 }
 
 void FLinkedFilteringMediator::BindCommands(const TSharedRef<FUICommandList>& InSequencerCommandList) const
 {
-	const TSharedRef<FSequencerFilterBar> SharedFilterBar = SequencerFilterViewModel->GetLinkedFilterBarImpl();
+	const TSharedRef<FSequencerFilterBar> SharedFilterBar = GetFilterModel()->GetLinkedFilterBarImpl();
 	SharedFilterBar->BindCommands();
 	
-	const TSharedRef<FSequencerFilterBar> InstancedFilterBar = SequencerFilterViewModel->GetInstancedFilterBarImpl();
+	const TSharedRef<FSequencerFilterBar> InstancedFilterBar = GetFilterModel()->GetInstancedFilterBarImpl();
 	InstancedFilterBar->BindCommands();
 	
 	InSequencerCommandList->MapAction(
 		FSequencerTrackFilterCommands::Get().ToggleLinkedFiltering ,
-		FExecuteAction::CreateLambda([WeakViewModel = SequencerFilterViewModel.ToWeakPtr()] { ToggleLinkedFiltering(WeakViewModel); }),
+		FExecuteAction::CreateLambda([WeakViewModel = GetFilterModel().ToWeakPtr()] { ToggleLinkedFiltering(WeakViewModel); }),
 		FCanExecuteAction::CreateLambda([]{ return IsLinkedFilteringFeatureEnabled(); }),
-		FIsActionChecked::CreateLambda([WeakViewModel = SequencerFilterViewModel.ToWeakPtr()] { return IsLinkedFilteringActionChecked(WeakViewModel); })
+		FIsActionChecked::CreateLambda([WeakViewModel = GetFilterModel().ToWeakPtr()] { return IsLinkedFilteringActionChecked(WeakViewModel); })
 		);
+	
+	InSequencerCommandList->Append(GetFilterAreaCommandList());
 }
 
 bool FLinkedFilteringMediator::ProcessCommandBindings(const FKeyEvent& InKeyEvent) const
 {
-	const TSharedPtr<FUICommandList> CommandList = SequencerFilterViewModel->GetActiveFilterBar()->GetCommandList();
-	return CommandList && CommandList->ProcessCommandBindings(InKeyEvent);
+	if (GetFilterAreaCommandList()->ProcessCommandBindings(InKeyEvent))
+	{
+		return true;
+	}
+	
+	const TSharedPtr<FUICommandList> ActiveFilterCommandList = GetFilterModel()->GetActiveFilterBar()->GetCommandList();
+	return ActiveFilterCommandList && ActiveFilterCommandList->ProcessCommandBindings(InKeyEvent);
 }
 
 ELinkedFilterMode FLinkedFilteringMediator::GetFilterMode() const
 {
-	return SequencerFilterViewModel->GetFilterMode();
+	return GetFilterModel()->GetFilterMode();
 }
 
 void FLinkedFilteringMediator::OnFilterModeChanged()
 {
 	// TODO: The goal of this is to call SetFilteredOut on all items again, but it does more work then it needs to. Instead of rerunning the filter,
 	// we could iterate through the hierarchy, and use FSequencerFilterBar::GetFilterData to call SetFilteredOut with the correct filter state.
-	SequencerFilterViewModel->GetActiveFilterBar()->RequestFilterUpdate();
+	GetFilterModel()->GetActiveFilterBar()->RequestFilterUpdate();
 }
 
 void FLinkedFilteringMediator::OnFilterStateChanged_LinkedFilterBar(
diff --git a/Engine/Source/Editor/Sequencer/Private/MVVM/ViewModels/LinkedFilteringMediator.h b/Engine/Source/Editor/Sequencer/Private/MVVM/ViewModels/LinkedFilteringMediator.h
--- a/Engine/Source/Editor/Sequencer/Private/MVVM/ViewModels/LinkedFilteringMediator.h
+++ b/Engine/Source/Editor/Sequencer/Private/MVVM/ViewModels/LinkedFilteringMediator.h
@@ -2,6 +2,7 @@
 
 #pragma once
 
+#include "Filters/ViewModels/Linking/FilterAreaManager.h"
 #include "MVVM/Extensions/IOutlinerExtension.h"
 #include "MVVM/ViewModels/ViewModel.h"
 
@@ -20,12 +21,12 @@ enum class ELinkedFilterMode : uint8;
  * - Calling IOutlinerExtension::SetFilteredOut when the item's filter state changes
  * - Binding commands that are bound by FSequencerFilterBar and making sure correct commands are executed depending on which one is active.
  */
-class FLinkedFilteringMediator : public FViewModel
+class FLinkedFilteringMediator : public FFilterAreaManager
 {
 public:
 	
-	explicit FLinkedFilteringMediator(const TSharedRef<FLinkedFilterViewModel>& InSequencerFilterViewModel);
-	virtual ~FLinkedFilteringMediator() override;
+	explicit FLinkedFilteringMediator(TWeakPtr<FSequencer> InSequencer, const TSharedRef<FLinkedFilterViewModel>& InSequencerFilterViewModel);
+	~FLinkedFilteringMediator();
 	
 	/** Binds the filter bar commands. */
 	void BindCommands(const TSharedRef<FUICommandList>& InSequencerCommandList) const;
@@ -35,9 +36,6 @@ public:
 	
 private:
 	
-	/** Knows whether Sequencer is using linked or instanced filtering. */
-	const TSharedRef<FLinkedFilterViewModel> SequencerFilterViewModel;
-	
 	/** @return The filter mode that Sequencer is currently using. */
 	ELinkedFilterMode GetFilterMode() const;
 
diff --git a/Engine/Source/Editor/Sequencer/Private/MVVM/ViewModels/SequencerEditorViewModelImpl.cpp b/Engine/Source/Editor/Sequencer/Private/MVVM/ViewModels/SequencerEditorViewModelImpl.cpp
--- a/Engine/Source/Editor/Sequencer/Private/MVVM/ViewModels/SequencerEditorViewModelImpl.cpp
+++ b/Engine/Source/Editor/Sequencer/Private/MVVM/ViewModels/SequencerEditorViewModelImpl.cpp
@@ -2,12 +2,12 @@
 
 #include "SequencerEditorViewModelImpl.h"
 
+#include "Filters/FilterConfigIdentifiers.h"
 #include "Filters/ViewModels/Linking/LinkedFilterFactoryViewModel.h"
 #include "Filters/ViewModels/Linking/LinkedFilterViewModel.h"
 #include "Filters/ViewModels/Linking/MakeLinkedFilterViewModelArgs.h"
 #include "LinkedFilteringMediator.h"
 #include "Sequencer.h"
-#include "Filters/FilterConfigIdentifiers.h"
 
 namespace UE::Sequencer
 {
@@ -27,6 +27,7 @@ void FSequencerEditorViewModelImpl::InitializeEditorImpl()
 void FSequencerEditorViewModelImpl::InitLinkedFiltering()
 {
 	const TSharedPtr<FSequencer> Sequencer = GetSequencerImpl();
+	
 	// TODO UE-362151: Implement data migration.
 	// Before linked filtering, the filter state was saved into the section string "*SequencerSettings->GetName()", e.g. "LevelSequencerEditor".
 	// Going forward, we can separate linked and instanced, e.g. ""LevelSequencerEditor.Linked", "LevelSequencerEditor.MainInstance", respectively.
@@ -42,6 +43,6 @@ void FSequencerEditorViewModelImpl::InitLinkedFiltering()
 	SequenceFilterInstanceViewModel = LinkedFilterFactory->MakeFilteringModelImpl(
 		FMakeLinkedFilterViewModelArgs(ConfigIds::Sequencer_InstancedConfigSubId)
 		);
-	FilterMediator = MakeShared<FLinkedFilteringMediator>(SequenceFilterInstanceViewModel.ToSharedRef());
+	FilterMediator = MakeShared<FLinkedFilteringMediator>(Sequencer, SequenceFilterInstanceViewModel.ToSharedRef());
 }
 }
diff --git a/Engine/Source/Editor/Sequencer/Private/SSequencer.cpp b/Engine/Source/Editor/Sequencer/Private/SSequencer.cpp
--- a/Engine/Source/Editor/Sequencer/Private/SSequencer.cpp
+++ b/Engine/Source/Editor/Sequencer/Private/SSequencer.cpp
@@ -157,6 +157,7 @@
 #include "Widgets/Sidebar/SequencerSelectionDrawer.h"
 #include "WorkspaceMenuStructure.h"
 #include "WorkspaceMenuStructureModule.h"
+#include "Filters/FilterConfigIdentifiers.h"
 #include "Filters/Widgets/Linking/SFilterBarSwitcher.h"
 #include "MVVM/ViewModels/LinkedFilteringMediator.h"
 #include "MVVM/ViewModels/SequencerEditorViewModelImpl.h"
@@ -364,7 +365,6 @@ void SSequencer::Construct(const FArguments& InArgs, TSharedRef<FSequencer> InSe
 	// TODO UE-363287
 	if (const TSharedPtr<FSequencerFilterBar> FilterBar = GetFilterBar())
 	{
-		FilterBar->OnVisibilityChanged().AddSP(this, &SSequencer::OnFilterBarVisibilityChanged);
 		FilterBar->OnFiltersChanged().AddSP(this, &SSequencer::OnTrackFiltersChanged);
 	}
 
@@ -4046,40 +4046,6 @@ TSharedPtr<FSequencerFilterBar> SSequencer::GetFilterBar() const
 	return nullptr;
 }
 
-bool SSequencer::IsFilterBarVisible() const
-{
-	if (!FilterBarWidget.IsValid())
-	{
-		return false;
-	}
-	
-	USequencerSettings* const SequencerSettings = GetSequencerSettings();
-	if (ensure(SequencerSettings) && !SequencerSettings->IsFilterBarVisible())
-	{
-		return false;
-	}
-
-	const TSharedPtr<FSequencerFilterBar> FilterBar = GetFilterBar();
-	return FilterBar && FilterBar->HasAnyFiltersEnabled();
-}
-
-void SSequencer::ToggleFilterBarVisibility()
-{
-	USequencerSettings* const SequencerSettings = GetSequencerSettings();
-	if (ensure(SequencerSettings))
-	{
-		const bool bNewFilterBarVisible = !SequencerSettings->IsFilterBarVisible();
-		SequencerSettings->SetFilterBarVisible(bNewFilterBarVisible);
-	}
-
-	RebuildFilterBarContent();
-}
-
-void SSequencer::OnFilterBarVisibilityChanged(const bool bInIsVisible)
-{
-	RebuildFilterBarContent();
-}
-
 void SSequencer::OnTrackFiltersChanged(const ESequencerFilterChange InChangeType, const TSharedRef<FSequencerTrackFilter>& InFilter)
 {
 	RebuildFilterBarContent();
@@ -4112,19 +4078,17 @@ void SSequencer::RebuildSearchAndFilterRow()
 	
 	const TSharedPtr<FSequencer> SequencerPin = SequencerPtr.Pin();
 	check(SequencerPin);
-	FilterBarWidget = SNew(SFilterBarSwitcher)
-		.AddMetaData<FTagMetaData>(FTagMetaData(TEXT("SequencerTrackFilters")))
-		.LinkedFilterViewModel(SequencerPin->GetViewModelImpl()->GetLinkedFilterViewModelImpl());
 
-	if (IsFilterBarVisible())
-	{
-		SearchAndFilterRow->AddSlot()
-			.AutoHeight()
-			.Padding(0)
-			[
-				FilterBarWidget.ToSharedRef()
-			];
-	}
+	SearchAndFilterRow->AddSlot()
+	.AutoHeight()
+	.Padding(0)
+	[
+		SNew(SFilterBarSwitcher)
+		.AddMetaData<FTagMetaData>(FTagMetaData(TEXT("SequencerTrackFilters")))
+		.Sequencer(SequencerPin)
+		.FilterAreaConfigId(ConfigIds::FilterArea_Outliner)
+		.LinkedFilterViewModel(SequencerPin->GetViewModelImpl()->GetLinkedFilterViewModelImpl())
+	];
 }
 
 TSharedPtr<FTabManager> SSequencer::GetTabManager() const
diff --git a/Engine/Source/Editor/Sequencer/Private/SSequencer.h b/Engine/Source/Editor/Sequencer/Private/SSequencer.h
--- a/Engine/Source/Editor/Sequencer/Private/SSequencer.h
+++ b/Engine/Source/Editor/Sequencer/Private/SSequencer.h
@@ -685,9 +685,6 @@ public:
 	
 	TSharedPtr<FSequencerFilterBar> GetFilterBar() const;
 
-	bool IsFilterBarVisible() const;
-	void ToggleFilterBarVisibility();
-
 	bool IsTransportControlsPoppedOut() const;
 	void SetTransportControlsPoppedOut(const bool bInPoppedOut);
 	void ToggleTransportControlsPoppedOut();
@@ -705,8 +702,6 @@ private:
 
 	void OnSidebarStateChanged(const FSidebarState& InNewState);
 
-	void OnFilterBarVisibilityChanged(bool bInIsVisible);
-
 	void OnTrackFiltersChanged(const ESequencerFilterChange InChangeType, const TSharedRef<FSequencerTrackFilter>& InFilter);
 
 	TSharedPtr<FTabManager> GetTabManager() const;
@@ -893,8 +888,6 @@ private:
 
 	TSharedPtr<FSequencerSelectionDrawer> SelectionDetailsDrawer;
 
-	/** Displays the filter pills. */
-	TSharedPtr<UE::Sequencer::SFilterBarSwitcher> FilterBarWidget;
 	/** Contains the filters widgets plus the filter bar. */
 	TSharedPtr<SVerticalBox> SearchAndFilterRow;
 	/** The widgets displayed above the tree view. */
diff --git a/Engine/Source/Editor/Sequencer/Private/SequencerSettings.cpp b/Engine/Source/Editor/Sequencer/Private/SequencerSettings.cpp
--- a/Engine/Source/Editor/Sequencer/Private/SequencerSettings.cpp
+++ b/Engine/Source/Editor/Sequencer/Private/SequencerSettings.cpp
@@ -1257,6 +1257,34 @@ bool USequencerSettings::RemoveTrackFilterBar(const FName InIdentifier)
 	return RemovedCount > 0;
 }
 
+FSequencerFilterAreaConfig& USequencerSettings::FindOrAddTrackFilterArea(const FName InIdentifier, bool bInSaveConfig)
+{
+	FSequencerFilterAreaConfig& FilterBarSettings = TrackFilterAreas.FindOrAdd(InIdentifier);
+
+	if (bInSaveConfig)
+	{
+		SaveConfig();
+	}
+
+	return FilterBarSettings;
+}
+
+FSequencerFilterAreaConfig* USequencerSettings::FindTrackFilterArea(const FName InIdentifier)
+{
+	return TrackFilterAreas.Find(InIdentifier);
+}
+
+bool USequencerSettings::RemoveTrackFilterArea(const FName InIdentifier)
+{
+	const int32 RemovedCount = TrackFilterBars.Remove(InIdentifier) > 0;
+	const bool bHasChanged = RemovedCount > 0;
+	if (bHasChanged)
+	{
+		SaveConfig();
+	}
+	return bHasChanged;
+}
+
 bool USequencerSettings::GetIncludePinnedInFilter() const
 {
 	return bIncludePinnedInFilter;
diff --git a/Engine/Source/Editor/Sequencer/Public/Filters/ISequencerFilterBar.h b/Engine/Source/Editor/Sequencer/Public/Filters/ISequencerFilterBar.h
--- a/Engine/Source/Editor/Sequencer/Public/Filters/ISequencerFilterBar.h
+++ b/Engine/Source/Editor/Sequencer/Public/Filters/ISequencerFilterBar.h
@@ -22,7 +22,6 @@ public:
 	
 	// This type is deprecated. However, DECLARE_EVENT_TwoParams does not support UE_DEPRECATED macro in Linux.
 	DECLARE_EVENT_TwoParams(ISequencerFilterBar, FOnFilterBarStateChanged, const bool /*InIsVisible*/, const EFilterBarLayout /*InNewLayout*/);
-	DECLARE_MULTICAST_DELEGATE_OneParam(FOnVisibilityChanged, bool /*bIsVisible*/)
 	DECLARE_MULTICAST_DELEGATE_OneParam(FMuteStateChangedDelegate, bool /* bNewIsMuted */);
 	
 	virtual ~ISequencerFilterBar() = default;
@@ -94,9 +93,6 @@ public:
 	PRAGMA_DISABLE_DEPRECATION_WARNINGS
 	virtual FOnFilterBarStateChanged& OnStateChanged() { return OnFilterBarStateChanged_DEPRECATED; }
 	PRAGMA_ENABLE_DEPRECATION_WARNINGS
-	
-	/** Event called when the visibility has changed. */
-	virtual FOnVisibilityChanged& OnVisibilityChanged() = 0;
 
 	/** Event called to request that the subscribe update its view since the filters state has changed */
 	virtual FSimpleMulticastDelegate& OnRequestUpdate() = 0;
@@ -107,10 +103,12 @@ public:
 	/** Event called when the mute state of the filters has changed. */
 	virtual FMuteStateChangedDelegate& OnMuteFiltersChanged() = 0;
 
-	virtual bool ShouldShowFilterBarWidget() const = 0;
-
-	virtual bool IsFilterBarVisible() const = 0;
-	virtual void ToggleFilterBarVisibility() = 0;
+	UE_DEPRECATED(5.8, "This function was removed")
+	virtual bool ShouldShowFilterBarWidget() const { return false; }
+	UE_DEPRECATED(5.8, "This function was removed")
+	virtual bool IsFilterBarVisible() const { return false; }
+	UE_DEPRECATED(5.8, "This function was removed")
+	virtual void ToggleFilterBarVisibility() {}
 
 	UE_DEPRECATED(5.8, "This function was removed")
 	virtual bool IsFilterBarLayout(const EFilterBarLayout InLayout) const { return false; }
diff --git a/Engine/Source/Editor/Sequencer/Public/Filters/SequencerFilterAreaConfig.h b/Engine/Source/Editor/Sequencer/Public/Filters/SequencerFilterAreaConfig.h
new file mode 100644
--- /dev/null
+++ b/Engine/Source/Editor/Sequencer/Public/Filters/SequencerFilterAreaConfig.h
@@ -0,0 +1,30 @@
+﻿// Copyright Epic Games, Inc. All Rights Reserved.
+
+#pragma once
+
+#include "SequencerFilterAreaConfig.generated.h"
+
+#define UE_API SEQUENCER_API
+
+/** 
+ * Config data that applies the entire filter area, regardless whether the filter bar displayed is in instanced or linked mode.
+ * @see FSequencerFilterBarConfig
+ */
+USTRUCT()
+struct FSequencerFilterAreaConfig
+{
+	GENERATED_BODY()
+	
+	UE_API void SetIsFilterBarVisible(bool bNewIsVisible);
+	UE_API bool IsFilterBarVisible() const;
+	
+private:
+	
+	/** Whether the filter bar should be visible. */
+	UPROPERTY()
+	bool bIsFilterBarVisible = false;
+	
+	// TODO UE-362150: safe whether the filter bar should be instanced or linked
+};
+
+#undef UE_API
diff --git a/Engine/Source/Editor/Sequencer/Public/Filters/SequencerFilterBarConfig.h b/Engine/Source/Editor/Sequencer/Public/Filters/SequencerFilterBarConfig.h
--- a/Engine/Source/Editor/Sequencer/Public/Filters/SequencerFilterBarConfig.h
+++ b/Engine/Source/Editor/Sequencer/Public/Filters/SequencerFilterBarConfig.h
@@ -22,6 +22,10 @@ public:
 	FString TextFilterString;
 };
 
+/** 
+ * Config data for a single filter bar instance.
+ * @see FSequencerFilterAreaConfig
+ */
 USTRUCT()
 struct FSequencerFilterBarConfig
 {
diff --git a/Engine/Source/Editor/Sequencer/Public/SequencerSettings.h b/Engine/Source/Editor/Sequencer/Public/SequencerSettings.h
--- a/Engine/Source/Editor/Sequencer/Public/SequencerSettings.h
+++ b/Engine/Source/Editor/Sequencer/Public/SequencerSettings.h
@@ -5,6 +5,7 @@
 #include "CoreMinimal.h"
 #include "CurveEditorZoomScaleConfig.h"
 #include "Filters/SequencerFilterBarConfig.h"
+#include "Filters/SequencerFilterAreaConfig.h"
 #include "Misc/SequencerThumbnailCaptureSettings.h"
 #include "UObject/ObjectMacros.h"
 #include "UObject/Object.h"
@@ -578,6 +579,11 @@ public:
 	UE_API FSequencerFilterBarConfig& FindOrAddTrackFilterBar(const FName InIdentifier, const bool bInSaveConfig);
 	UE_API FSequencerFilterBarConfig* FindTrackFilterBar(const FName InIdentifier);
 	UE_API bool RemoveTrackFilterBar(const FName InIdentifier);
+	
+	UE_API FSequencerFilterAreaConfig& FindOrAddTrackFilterArea(const FName InIdentifier, bool bInSaveConfig = true);
+	UE_API FSequencerFilterAreaConfig* FindTrackFilterArea(const FName InIdentifier);
+	UE_API bool RemoveTrackFilterArea(const FName InIdentifier);
+	
 
 	UE_API bool GetIncludePinnedInFilter() const;
 	UE_API void SetIncludePinnedInFilter(const bool bInIncludePinned);
@@ -588,7 +594,9 @@ public:
 	UE_API bool GetUseFilterSubmenusForCategories() const;
 	UE_API void SetUseFilterSubmenusForCategories(const bool bInUseFilterSubmenusForCategories);
 
+	UE_DEPRECATED(5.8, "Use the version that takes in an identifier instead.")
 	UE_API bool IsFilterBarVisible() const;
+	UE_DEPRECATED(5.8, "Use the version that takes in an identifier instead.")
 	UE_API void SetFilterBarVisible(const bool bInVisible);
 
 	UE_DEPRECATED(5.8, "Sequencer now only uses horizontal layout")
@@ -951,9 +959,20 @@ protected:
 	UPROPERTY(config)
 	TMap<FName, FSidebarState> SidebarState;
 
-	/** Saved settings for each unique filter bar instance mapped by instance identifier */
+	/** 
+	 * Saved settings for each unique filter bar instance mapped by instance identifier.
+	 * Filter bar instances live in a filter bar areas. 
+	 * @see TrackFilterAreas For settings independent whether the filter bar is in linked or instanced mode, such as visibility.
+	 */
 	UPROPERTY(config, EditAnywhere, Category = Filtering)
 	TMap<FName, FSequencerFilterBarConfig> TrackFilterBars;
+	
+	/**
+	 * Saved settings for each filter UI area. 
+	 * @see TrackFilterBars For settings that are specific to the linked or instanced mode, such as the filter active state.
+	 */
+	UPROPERTY(config, EditAnywhere, Category = Filtering)
+	TMap<FName, FSequencerFilterAreaConfig> TrackFilterAreas;
 
 	/** Apply filtering to pinned tracks that would otherwise ignore filters */
 	UPROPERTY(config, EditAnywhere, Category = Filtering)
