Sequencer: Adds view options for curve editor.
- Leverages FFilterAreaClient to mutate filter bar visbility for curve editor
- Adjusts FSequencerViewOptionsMenu so the contents of view options can be customized for curve editor and outliner
- View options content is now injected into SSequencerViewOptionsButton

#Synced-CL 51334859

diff --git a/Engine/Source/Editor/Sequencer/Private/CurveEditor/Integration/CurveEditorCommandBinder.cpp b/Engine/Source/Editor/Sequencer/Private/CurveEditor/Integration/CurveEditorCommandBinder.cpp
--- a/Engine/Source/Editor/Sequencer/Private/CurveEditor/Integration/CurveEditorCommandBinder.cpp
+++ b/Engine/Source/Editor/Sequencer/Private/CurveEditor/Integration/CurveEditorCommandBinder.cpp
@@ -3,29 +3,30 @@
 #include "CurveEditorCommandBinder.h"
 
 #include "CurveEditor.h"
+#include "Filters/SequencerTrackFilterCommands.h"
+#include "Filters/ViewModels/Linking/FilterAreaManager.h"
+#include "Filters/ViewModels/Linking/LinkedFilterViewModel.h"
 #include "ISequencer.h"
 #include "Sequencer.h"
 #include "SequencerCommands.h"
-#include "Filters/SequencerTrackFilterCommands.h"
-#include "Filters/ViewModels/Linking/LinkedFilterViewModel.h"
 
 namespace UE::Sequencer
 {
 FCurveCommandBinder::FCurveCommandBinder(
 	const TSharedRef<FSequencer>& InSequencer,
 	const TSharedRef<FCurveEditor>& InCurveEditor,
-	const TSharedRef<FLinkedFilterViewModel>& InFilterViewModel
+	const FFilterAreaManager& InFilterArea
 	)
 	: WeakSequencer(InSequencer)
 	, CurveEditor(InCurveEditor)
-	, FilterViewModel(InFilterViewModel)
+	, FilterArea(InFilterArea)
 {
-	FilterViewModel->OnFilterModeChanged().AddRaw(this, &FCurveCommandBinder::OnFilterModeChanged);
+	FilterArea.GetFilterModel()->OnFilterModeChanged().AddRaw(this, &FCurveCommandBinder::OnFilterModeChanged);
 }
 
 FCurveCommandBinder::~FCurveCommandBinder()
 {
-	FilterViewModel->OnFilterModeChanged().RemoveAll(this);
+	FilterArea.GetFilterModel()->OnFilterModeChanged().RemoveAll(this);
 }
 
 void FCurveCommandBinder::InitSequencerCommandBindings()
@@ -36,7 +37,7 @@ void FCurveCommandBinder::InitSequencerCommandBindings()
 		return;
 	}
 	
-	const TSharedPtr<FUICommandList> CurveEditorSharedBindings = Sequencer->GetCommandBindings(ESequencerCommandBindings::CurveEditor);
+	const TSharedPtr<FUICommandList> CurveEditorSharedBindings = GetCurveEditorCommands();
 	const TSharedPtr<FUICommandList> SequencerCommandBindings = Sequencer->GetCommandBindings(ESequencerCommandBindings::Sequencer);
 	
 	const FSequencerCommands& Commands = FSequencerCommands::Get();
@@ -69,11 +70,18 @@ void FCurveCommandBinder::InitSequencerCommandBindings()
 	CurveEditorSharedBindings->MapAction(Commands.AddRotationKey, *SequencerCommandBindings->GetActionForCommand(Commands.AddRotationKey));
 	CurveEditorSharedBindings->MapAction(Commands.AddScaleKey, *SequencerCommandBindings->GetActionForCommand(Commands.AddScaleKey));
 	
+	FilterArea.GetFilterAreaCommandList()->Append(CurveEditorSharedBindings.ToSharedRef());
 	CurveEditor->GetCommands()->Append(CurveEditorSharedBindings.ToSharedRef());
 	
 	BindCurveEditorToActiveFilterBar();
 }
 
+TSharedPtr<FUICommandList> FCurveCommandBinder::GetCurveEditorCommands() const
+{
+	const TSharedPtr<FSequencer> Sequencer = WeakSequencer.Pin();
+	return Sequencer ? Sequencer->GetCommandBindings(ESequencerCommandBindings::CurveEditor) : nullptr;
+}
+
 void FCurveCommandBinder::BindCurveEditorToActiveFilterBar() const
 {
 	const TSharedPtr<FSequencer> Sequencer = WeakSequencer.Pin();
@@ -83,7 +91,7 @@ void FCurveCommandBinder::BindCurveEditorToActiveFilterBar() const
 		return;
 	}
 	
-	const TSharedPtr<FSequencerFilterBar> FilterBar = FilterViewModel->GetActiveFilterBarImpl();
+	const TSharedPtr<FSequencerFilterBar> FilterBar = FilterArea.GetFilterModel()->GetActiveFilterBarImpl();
 	const TSharedPtr<FUICommandList> FilterBarCommandList = FilterBar->GetCommandList();
 	
 	// Add the general track filter commands
diff --git a/Engine/Source/Editor/Sequencer/Private/CurveEditor/Integration/CurveEditorCommandBinder.h b/Engine/Source/Editor/Sequencer/Private/CurveEditor/Integration/CurveEditorCommandBinder.h
--- a/Engine/Source/Editor/Sequencer/Private/CurveEditor/Integration/CurveEditorCommandBinder.h
+++ b/Engine/Source/Editor/Sequencer/Private/CurveEditor/Integration/CurveEditorCommandBinder.h
@@ -2,15 +2,18 @@
 
 #pragma once
 
+#include "Filters/ViewModels/Linking/FilterAreaManager.h"
 #include "Templates/SharedPointer.h"
 #include "Templates/UnrealTemplate.h"
 
+class FUICommandList;
 class FCurveEditor;
 class FSequencer;
 
 namespace UE::Sequencer
 {
 class FLinkedFilterViewModel;
+
 /** 
  * Injects some Sequencer commands into Curve Editor such that when Curve Editor UI is focused, the user can continue to use some of the
  * Sequencer commands, and they'll take affects as if they were issued while Sequencer UI was focused.
@@ -21,13 +24,16 @@ public:
 	
 	explicit FCurveCommandBinder(
 		const TSharedRef<FSequencer>& InSequencer, const TSharedRef<FCurveEditor>& InCurveEditor, 
-		const TSharedRef<FLinkedFilterViewModel>& InFilterViewModel
+		const FFilterAreaManager& InFilterArea UE_LIFETIMEBOUND
 		);
 	~FCurveCommandBinder();
 	
 	/** Binds Sequencer commands that should be exposed in Curve Editor. */
 	void InitSequencerCommandBindings();
 	
+	/** @return The command list that is used in the UI of Sequencer's curve editor. */
+	TSharedPtr<FUICommandList> GetCurveEditorCommands() const;
+	
 private:
 	
 	/** The owning Sequencer whose commands are being rebound and that owns the FCurveEditor. */
@@ -37,8 +43,7 @@ private:
 	const TSharedRef<FCurveEditor> CurveEditor;
 	
 	/** When the ELinkedFilterMode changes, the bindings in the active filter bar must be reapplied to Curve Editor.*/
-	const TSharedRef<FLinkedFilterViewModel> FilterViewModel;
-	
+	const FFilterAreaManager& FilterArea;
 	
 	/** Handles any command rebinding that must take place when the active filter bar changes. */
 	void OnFilterModeChanged() const { BindCurveEditorToActiveFilterBar(); }
diff --git a/Engine/Source/Editor/Sequencer/Private/CurveEditor/SequencerCurveEditorApp.cpp b/Engine/Source/Editor/Sequencer/Private/CurveEditor/SequencerCurveEditorApp.cpp
--- a/Engine/Source/Editor/Sequencer/Private/CurveEditor/SequencerCurveEditorApp.cpp
+++ b/Engine/Source/Editor/Sequencer/Private/CurveEditor/SequencerCurveEditorApp.cpp
@@ -3,6 +3,7 @@
 #include "SequencerCurveEditorApp.h"
 
 #include "Filters/FilterConfigIdentifiers.h"
+#include "Filters/ViewModels/Linking/FilterAreaManager.h"
 #include "Filters/ViewModels/Linking/LinkedFilterFactoryViewModel.h"
 #include "Filters/ViewModels/Linking/LinkedFilterViewModel.h"
 #include "Filters/ViewModels/Linking/MakeLinkedFilterViewModelArgs.h"
@@ -14,6 +15,7 @@
 #include "Specializations/SequencerCurveEditorToolbarExtender.h"
 #include "Specializations/SequencerCurveModelUtils.h"
 #include "TimeSliderArgs.h"
+#include "Views/CurveEditorWidgetOwnerArgs.h"
 #include "Widgets/SequencerCurveEditorTimeSliderController.h"
 
 namespace UE::Sequencer
@@ -67,17 +69,8 @@ static TSharedRef<FCurveEditor> MakeCurveEditor(const TSharedRef<FSequencer>& In
 	return CurveEditorModel;
 }
 
-static TSharedRef<FLinkedFilterViewModel> MakeFilteringModel(
-	const TSharedRef<FSequencer>& InSequencer, const FTimeSliderArgs& InTimeSliderArgs, const TSharedRef<FCurveEditor>& InCurveEditor
-	)
+static TSharedRef<FLinkedFilterViewModel> MakeFilteringModel(const TSharedRef<FSequencer>& InSequencer)
 {
-	// We create a custom Time Slider Controller which is just a wrapper around the actual one, but is 
-	// aware of our custom bounds logic. Currently the range the bar displays is tied to Sequencer 
-	// timeline and not the Bounds, so we need a way of changing it to look at the Bounds but only for 
-	// the Curve Editor time slider controller. We want everything else to just pass through though.
-	const TSharedRef<ITimeSliderController> CurveEditorTimeSliderController = MakeShared<FSequencerCurveEditorTimeSliderController>(
-		InTimeSliderArgs, InSequencer, InCurveEditor);
-	
 	const TSharedPtr<FLinkedFilterFactoryViewModel> FilterFactory = InSequencer->GetViewModelImpl()->GetLinkedFilterViewModelFactoryImpl();
 	const TSharedRef<FLinkedFilterViewModel> ViewModelImpl = FilterFactory->MakeFilteringModelImpl(
 		FMakeLinkedFilterViewModelArgs(ConfigIds::CurveEditor_InstancedConfigId)
@@ -89,13 +82,18 @@ static TSharedRef<FLinkedFilterViewModel> MakeFilteringModel(
 
 FSequencerCurveEditorApp::FSequencerCurveEditorApp(const TSharedRef<FSequencer>& InSequencer, const FTimeSliderArgs& InTimeSliderArgs)
 	: CurveEditor(CurveEditorDetail::MakeCurveEditor(InSequencer, InTimeSliderArgs))
-	, LinkedFilteringViewModel(CurveEditorDetail::MakeFilteringModel(InSequencer, InTimeSliderArgs, CurveEditor))
-	, CommandBinder(InSequencer, CurveEditor, LinkedFilteringViewModel)
-	, SelectionSyncLogic(InSequencer, CurveEditor, LinkedFilteringViewModel)
-	, WidgetOwner(
+	, FilterAreaClient(
+		ConfigIds::FilterArea_CurveEditor,
+		InSequencer,
+		CurveEditorDetail::MakeFilteringModel(InSequencer)
+	)
+	, CommandBinder(InSequencer, CurveEditor, FilterAreaClient)
+	, SelectionSyncLogic(InSequencer, CurveEditor, FilterAreaClient.GetFilterModel())
+	, WidgetOwner(FCurveEditorWidgetOwnerArgs(
 		InSequencer, CurveEditor, 
 		MakeShared<FSequencerCurveEditorTimeSliderController>(InTimeSliderArgs, InSequencer, CurveEditor), 
-		LinkedFilteringViewModel, SelectionSyncLogic.GetSequencerToCurveEditorSyncer()
+		FilterAreaClient.GetFilterModel(), CommandBinder.GetCurveEditorCommands().ToSharedRef(),
+		SelectionSyncLogic.GetSequencerToCurveEditorSyncer())
 		)
 {
 	SelectionSyncLogic.OnRequestRefreshCurveEditorUI().AddLambda([this]
diff --git a/Engine/Source/Editor/Sequencer/Private/CurveEditor/SequencerCurveEditorApp.h b/Engine/Source/Editor/Sequencer/Private/CurveEditor/SequencerCurveEditorApp.h
--- a/Engine/Source/Editor/Sequencer/Private/CurveEditor/SequencerCurveEditorApp.h
+++ b/Engine/Source/Editor/Sequencer/Private/CurveEditor/SequencerCurveEditorApp.h
@@ -3,6 +3,7 @@
 #pragma once
 
 #include "CurveEditor/Views/CurveEditorWidgetOwner.h"
+#include "Filters/ViewModels/Linking/FilterAreaManager.h"
 #include "Integration/CurveEditorCommandBinder.h"
 #include "Integration/Selection/SelectionSyncRootLogic.h"
 #include "Templates/SharedPointer.h"
@@ -65,14 +66,14 @@ public:
 	
 private:
 	
-	// @note Member order matters. C++ destruction order is reverse declaration order, so the dependencies must points upwards:
+	// @note Member order matters. C++ destruction order is reverse declaration order, so the dependencies must point upwards:
 	// Lower members can only depend on members declared after. Higher members cannot depend on members declared after.
 	
 	/** The Curve Editor instance that Sequencer uses. */
 	TSharedRef<FCurveEditor> CurveEditor;
 	
-	/** Representation of Curve Editor's linked filters. */
-	TSharedRef<FLinkedFilterViewModel> LinkedFilteringViewModel;
+	/** Manages linked filtering for curve editor. Owns the view model for linked filtering. */
+	FFilterAreaManager FilterAreaClient;
 	
 	/** Allows users to use issue some Sequencer commands when the Curve Editor is UI is focused. */
 	FCurveCommandBinder CommandBinder;
diff --git a/Engine/Source/Editor/Sequencer/Private/CurveEditor/Views/CurveEditorWidgetOwner.cpp b/Engine/Source/Editor/Sequencer/Private/CurveEditor/Views/CurveEditorWidgetOwner.cpp
--- a/Engine/Source/Editor/Sequencer/Private/CurveEditor/Views/CurveEditorWidgetOwner.cpp
+++ b/Engine/Source/Editor/Sequencer/Private/CurveEditor/Views/CurveEditorWidgetOwner.cpp
@@ -2,6 +2,7 @@
 
 #include "CurveEditorWidgetOwner.h"
 
+#include "CurveEditorWidgetOwnerArgs.h"
 #include "CurveEditor/Widgets/SSequencerCurveEditor.h"
 #include "CurveEditor/Widgets/Tree/SCurveTreeContent.h"
 #include "CurveEditor/Widgets/Tree/SOldCurveTreeContent.h"
@@ -34,29 +35,25 @@ static TAutoConsoleVariable<bool> CVarUseLinkedFilterUI(
 	);
 
 static FCurveEditorWidgets MakePresenterWidgets(
-	const TSharedRef<FSequencer>& InSequencer, 
-	const TSharedRef<FCurveEditor>& InCurveEditor, 
-	const TSharedRef<ITimeSliderController>& InTimeSliderController,
-	const TSharedRef<FLinkedFilterViewModel>& InFilteringViewModel,
-	const FCurveModelSyncer& InCurveModelSyncer UE_LIFETIMEBOUND,
+	const FCurveEditorWidgetOwnerArgs& InArgs,
 	const TSharedPtr<FTabManager> InTabManager
 	)
 {
-	const TSharedRef<SOldCurveTreeContent> OldContent = SNew(SOldCurveTreeContent, InSequencer, InCurveEditor)
+	const TSharedRef<SOldCurveTreeContent> OldContent = SNew(SOldCurveTreeContent, InArgs.Sequencer, InArgs.CurveEditor)
 		.TabManager(InTabManager);
 	const TSharedRef<SCurveTreeContent> TreeContent = SNew(SCurveTreeContent)
 		.TabManager(InTabManager)
-		.Sequencer(InSequencer)
-		.CurveEditor(InCurveEditor)
-		.FilterViewModel(InFilteringViewModel)
-		.CurveModelSyncer(&InCurveModelSyncer)
+		.Sequencer(InArgs.Sequencer)
+		.CurveEditor(InArgs.CurveEditor)
+		.FilterViewModel(InArgs.FilteringViewModel)
+		.CommandList(InArgs.CommandList)
+		.CurveModelSyncer(&InArgs.CurveModelSyncer)
 		.FilterPills()
 		[
 			SNew(SFilterBarSwitcher)
-			.Sequencer(InSequencer)
-			// TODO UE-368977: Introduce separate setting for curve editor filter bar visibility. For now, use the same as in Sequencer Outliner.
-			.FilterAreaConfigId(ConfigIds::FilterArea_Outliner)
-			.LinkedFilterViewModel(InFilteringViewModel)
+			.Sequencer(InArgs.Sequencer)
+			.FilterAreaConfigId(ConfigIds::FilterArea_CurveEditor)
+			.LinkedFilterViewModel(InArgs.FilteringViewModel)
 		];
 	
 	const TSharedRef<SWidgetSwitcher> Content = 
@@ -65,10 +62,10 @@ static FCurveEditorWidgets MakePresenterWidgets(
 		+SWidgetSwitcher::Slot()[ OldContent ]
 		+SWidgetSwitcher::Slot()[ TreeContent ];
 	
-	const TSharedRef<SCurveEditorPanel> CurveEditorPanel = SNew(SCurveEditorPanel, InCurveEditor)
+	const TSharedRef<SCurveEditorPanel> CurveEditorPanel = SNew(SCurveEditorPanel, InArgs.CurveEditor)
 		// Grid lines match the color specified in FSequencerTimeSliderController::OnPaintViewArea
 		.GridLineTint(FLinearColor(0.f, 0.f, 0.f, 0.3f))
-		.ExternalTimeSliderController(InTimeSliderController)
+		.ExternalTimeSliderController(InArgs.TimeSliderController)
 		.MinimumViewPanelHeight(0.f)
 		.TabManager(InTabManager)
 		.DisabledTimeSnapTooltip(LOCTEXT("CurveEditorTimeSnapDisabledTooltip", "Time Snapping is currently driven by Sequencer."))
@@ -77,7 +74,7 @@ static FCurveEditorWidgets MakePresenterWidgets(
 			Content
 		];
 	
-	const TSharedRef<SSequencerCurveEditor> RootWidget = SNew(SSequencerCurveEditor, CurveEditorPanel, InSequencer);
+	const TSharedRef<SSequencerCurveEditor> RootWidget = SNew(SSequencerCurveEditor, CurveEditorPanel, InArgs.Sequencer);
 	return FCurveEditorWidgets(RootWidget, CurveEditorPanel, OldContent, TreeContent);
 }
 }
@@ -87,21 +84,13 @@ bool UseLinkedFilteringUI()
 	return PresenterDetail::CVarUseLinkedFilterUI.GetValueOnAnyThread();
 }
 
-FCurveEditorWidgetOwner::FCurveEditorWidgetOwner(
-	const TSharedRef<FSequencer>& InSequencer,
-	const TSharedRef<FCurveEditor>& InCurveEditor,
-	const TSharedRef<ITimeSliderController>& InTimeSliderController,
-	const TSharedRef<FLinkedFilterViewModel>& InFilteringViewModel,
-	const FCurveModelSyncer& InCurveModelSyncer
-	)
-	: WeakSequencer(InSequencer)
-	, CurveEditor(InCurveEditor)
-	, Widgets(PresenterDetail::MakePresenterWidgets(
-		InSequencer, InCurveEditor, InTimeSliderController, InFilteringViewModel, InCurveModelSyncer, GetTabManager()
-		))
+FCurveEditorWidgetOwner::FCurveEditorWidgetOwner(const FCurveEditorWidgetOwnerArgs& InArgs)
+	: WeakSequencer(InArgs.Sequencer)
+	, CurveEditor(InArgs.CurveEditor)
+	, Widgets(PresenterDetail::MakePresenterWidgets(InArgs, GetTabManager()))
 {
 	BindCommands();
-	InitFrameNumberPropertyLayout(InSequencer);
+	InitFrameNumberPropertyLayout(InArgs.Sequencer);
 	
 	const TSharedPtr<FTabManager> TabManager = GetTabManager();
 	if (const TSharedPtr<SDockTab> ExistingCurveEditorTab = TabManager->FindExistingLiveTab(FCurveEditorExtension::CurveEditorTabName))
diff --git a/Engine/Source/Editor/Sequencer/Private/CurveEditor/Views/CurveEditorWidgetOwner.h b/Engine/Source/Editor/Sequencer/Private/CurveEditor/Views/CurveEditorWidgetOwner.h
--- a/Engine/Source/Editor/Sequencer/Private/CurveEditor/Views/CurveEditorWidgetOwner.h
+++ b/Engine/Source/Editor/Sequencer/Private/CurveEditor/Views/CurveEditorWidgetOwner.h
@@ -6,6 +6,7 @@
 #include "HAL/Platform.h"
 #include "Templates/SharedPointer.h"
 
+class FUICommandList;
 class FCurveEditor;
 class FSequencer;
 class FTabManager;
@@ -20,6 +21,7 @@ namespace UE::Sequencer
 {
 class FCurveModelSyncer;
 class FLinkedFilterViewModel;
+struct FCurveEditorWidgetOwnerArgs;
 
 /** @return Whether the linked filtering UI is enabled. */
 bool UseLinkedFilteringUI();
@@ -36,13 +38,7 @@ class FCurveEditorWidgetOwner
 {
 public:
 	
-	explicit FCurveEditorWidgetOwner(
-		const TSharedRef<FSequencer>& InSequencer, 
-		const TSharedRef<FCurveEditor>& InCurveEditor,
-		const TSharedRef<ITimeSliderController>& InTimeSliderController,
-		const TSharedRef<FLinkedFilterViewModel>& InFilteringViewModel,
-		const FCurveModelSyncer& InCurveModelSyncer UE_LIFETIMEBOUND
-		);
+	explicit FCurveEditorWidgetOwner(const FCurveEditorWidgetOwnerArgs& InArgs);
 	~FCurveEditorWidgetOwner();
 	
 	/** Returns whether the curve editor is open */
diff --git a/Engine/Source/Editor/Sequencer/Private/CurveEditor/Views/CurveEditorWidgetOwnerArgs.h b/Engine/Source/Editor/Sequencer/Private/CurveEditor/Views/CurveEditorWidgetOwnerArgs.h
new file mode 100644
--- /dev/null
+++ b/Engine/Source/Editor/Sequencer/Private/CurveEditor/Views/CurveEditorWidgetOwnerArgs.h
@@ -0,0 +1,45 @@
+﻿// Copyright Epic Games, Inc. All Rights Reserved.
+
+#pragma once
+
+#include "HAL/Platform.h"
+#include "Templates/SharedPointer.h"
+
+class FUICommandList;
+class ITimeSliderController;
+class FCurveEditor;
+class FSequencer;
+
+namespace UE::Sequencer
+{
+class FCurveModelSyncer;
+class FLinkedFilterViewModel;
+
+/** Init args for FCurveEditorWidgetOwner */
+struct FCurveEditorWidgetOwnerArgs
+{
+	TSharedRef<FSequencer> Sequencer; 
+	
+	TSharedRef<FCurveEditor> CurveEditor;
+	TSharedRef<ITimeSliderController> TimeSliderController;
+	
+	TSharedRef<FLinkedFilterViewModel> FilteringViewModel;
+	TSharedRef<FUICommandList> CommandList;
+	
+	FCurveModelSyncer& CurveModelSyncer;
+
+	FCurveEditorWidgetOwnerArgs(
+		const TSharedRef<FSequencer>& Sequencer, 
+		const TSharedRef<FCurveEditor>& CurveEditor, const TSharedRef<ITimeSliderController>& TimeSliderController, 
+		const TSharedRef<FLinkedFilterViewModel>& FilteringViewModel, const TSharedRef<FUICommandList>& CommandList,
+		FCurveModelSyncer& CurveModelSyncer UE_LIFETIMEBOUND
+		)
+		: Sequencer(Sequencer)
+		, CurveEditor(CurveEditor)
+		, TimeSliderController(TimeSliderController)
+		, FilteringViewModel(FilteringViewModel)
+		, CommandList(CommandList)
+		, CurveModelSyncer(CurveModelSyncer)
+	{}
+};
+} // namespace UE::Sequencer
diff --git a/Engine/Source/Editor/Sequencer/Private/CurveEditor/Widgets/SCurveEditorSearchAndFilterRow.cpp b/Engine/Source/Editor/Sequencer/Private/CurveEditor/Widgets/SCurveEditorSearchAndFilterRow.cpp
--- a/Engine/Source/Editor/Sequencer/Private/CurveEditor/Widgets/SCurveEditorSearchAndFilterRow.cpp
+++ b/Engine/Source/Editor/Sequencer/Private/CurveEditor/Widgets/SCurveEditorSearchAndFilterRow.cpp
@@ -3,17 +3,17 @@
 #include "SCurveEditorSearchAndFilterRow.h"
 
 #include "Filters/Menus/SequencerTrackFilterMenu.h"
+#include "Filters/Menus/SequencerViewOptionsMenu.h"
 #include "Filters/SequencerFilterBar.h"
 #include "Filters/Utils/HideIsolateViewModelUtils.h"
+#include "Filters/Utils/LinkedFilteringViewUtils.h"
 #include "Filters/ViewModels/HideIsolateShowViewModel.h"
 #include "Filters/ViewModels/Linking/ILinkedFilterViewModel.h"
 #include "Filters/ViewModels/Linking/LinkedFilterViewModel.h"
 #include "Filters/Widgets/Linking/SFilterLinkStateSwitcher.h"
 #include "Filters/Widgets/Linking/SLinkedFilteringToggleButton.h"
 #include "Filters/Widgets/SSearchAndFilterWidget.h"
-#include "MVVM/Selection/Selection.h"
 #include "Sequencer.h"
-#include "Filters/Utils/LinkedFilteringViewUtils.h"
 #include "Widgets/SBoxPanel.h"
 #include "Widgets/Sequencer/SSequencerViewOptionsButton.h"
 
@@ -37,10 +37,35 @@ void SCurveEditorSearchAndFilterRow::Construct(const FArguments& InArgs)
 {
 	WeakSequencer = InArgs._Sequencer;
 	FilterModel = InArgs._FilterViewModel;
+	check(InArgs._CommandList);
 	
 	ChildSlot
 	[
-		MakeFilterContent(InArgs)
+		SNew(SHorizontalBox)
+		
+		+ SHorizontalBox::Slot()
+		.FillWidth(1.f)
+		.VAlign(VAlign_Center)
+		.Padding(2.f, 0.f, 0.f, 0.f)
+		[
+			MakeFilterContent(InArgs)
+		]
+		
+		// View Options Combo Button
+		+ SHorizontalBox::Slot()
+		.AutoWidth()
+		.VAlign(VAlign_Center)
+		.Padding(2.f, 0.f, 0.f, 0.f)
+		[
+			SNew(SSequencerViewOptionsButton)
+			.OnGetMenuContent_Lambda([this, ViewOptions = MakeShared<FSequencerViewOptionsMenu>(), CommandList = InArgs._CommandList]
+			{
+				const TSharedPtr<FSequencer> Sequencer = WeakSequencer.Pin();
+				return Sequencer 
+					? ViewOptions->CreateMenuForCurveEditor(WeakSequencer, CommandList.ToSharedRef()) 
+					: SNullWidget::NullWidget;
+			})
+		]
 	];
 }
 
diff --git a/Engine/Source/Editor/Sequencer/Private/CurveEditor/Widgets/SCurveEditorSearchAndFilterRow.h b/Engine/Source/Editor/Sequencer/Private/CurveEditor/Widgets/SCurveEditorSearchAndFilterRow.h
--- a/Engine/Source/Editor/Sequencer/Private/CurveEditor/Widgets/SCurveEditorSearchAndFilterRow.h
+++ b/Engine/Source/Editor/Sequencer/Private/CurveEditor/Widgets/SCurveEditorSearchAndFilterRow.h
@@ -6,6 +6,7 @@
 #include "Widgets/DeclarativeSyntaxSupport.h"
 #include "Widgets/SCompoundWidget.h"
 
+class FUICommandList;
 class SSequencerSearchBox;
 class ISequencerTrackFilters;
 class FSequencer;
@@ -24,10 +25,12 @@ class SCurveEditorSearchAndFilterRow : public SCompoundWidget
 public:
 	
 	SLATE_BEGIN_ARGS(SCurveEditorSearchAndFilterRow) {}
-		/** Used to generate the subwidgets */
+		/** Required. Used to generate the subwidgets */
 		SLATE_ARGUMENT(TWeakPtr<FSequencer>, Sequencer)
-		/** Determines the filtering widgets displayed. */
+		/** Required. Determines the filtering widgets displayed. */
 		SLATE_ARGUMENT(TSharedPtr<FLinkedFilterViewModel>, FilterViewModel)
+		/** Required. The command list that should be used, e.g. by the view options context menu. */
+		SLATE_ARGUMENT(TSharedPtr<FUICommandList>, CommandList)
 
 		/** Gets all items displayed in the tree view. */
 		SLATE_ATTRIBUTE(TSet<TWeakViewModelPtr<IOutlinerExtension>>, AllItems)
diff --git a/Engine/Source/Editor/Sequencer/Private/CurveEditor/Widgets/Tree/SCurveTreeContent.cpp b/Engine/Source/Editor/Sequencer/Private/CurveEditor/Widgets/Tree/SCurveTreeContent.cpp
--- a/Engine/Source/Editor/Sequencer/Private/CurveEditor/Widgets/Tree/SCurveTreeContent.cpp
+++ b/Engine/Source/Editor/Sequencer/Private/CurveEditor/Widgets/Tree/SCurveTreeContent.cpp
@@ -87,6 +87,7 @@ TSharedRef<SWidget> SCurveTreeContent::MakeFilterArea(const FArguments& InArgs)
 				SAssignNew(FilterRow, SCurveEditorSearchAndFilterRow)
 				.Sequencer(InArgs._Sequencer)
 				.FilterViewModel(InArgs._FilterViewModel)
+				.CommandList(InArgs._CommandList)
 				.ScrollItemIntoView(this, &SCurveTreeContent::ScrollIntoView)
 				.AllItems(this, &SCurveTreeContent::GetAllItems)
 				.SelectedItems(this, &SCurveTreeContent::GetSelectedItems)
diff --git a/Engine/Source/Editor/Sequencer/Private/CurveEditor/Widgets/Tree/SCurveTreeContent.h b/Engine/Source/Editor/Sequencer/Private/CurveEditor/Widgets/Tree/SCurveTreeContent.h
--- a/Engine/Source/Editor/Sequencer/Private/CurveEditor/Widgets/Tree/SCurveTreeContent.h
+++ b/Engine/Source/Editor/Sequencer/Private/CurveEditor/Widgets/Tree/SCurveTreeContent.h
@@ -9,6 +9,7 @@
 class FCurveEditor;
 class FSequencer;
 class FTabManager;
+class FUICommandList;
 class SCurveEditorTree;
 template<typename T> struct STemporarilyFocusedSpinBox;
 
@@ -30,6 +31,8 @@ public:
 		SLATE_ARGUMENT(TSharedPtr<FCurveEditor>, CurveEditor)
 		/** Required. Determines the filtering widgets displayed. */
 		SLATE_ARGUMENT(TSharedPtr<FLinkedFilterViewModel>, FilterViewModel)
+		/** Required. The command list that should be used, e.g. by the view options context menu. */
+		SLATE_ARGUMENT(TSharedPtr<FUICommandList>, CommandList)
 		
 		/** Explicit tab manager to use, otherwise uses the global tab manager */
 		SLATE_ARGUMENT(TSharedPtr<FTabManager>, TabManager)
diff --git a/Engine/Source/Editor/Sequencer/Private/Filters/FilterConfigIdentifiers.cpp b/Engine/Source/Editor/Sequencer/Private/Filters/FilterConfigIdentifiers.cpp
--- a/Engine/Source/Editor/Sequencer/Private/Filters/FilterConfigIdentifiers.cpp
+++ b/Engine/Source/Editor/Sequencer/Private/Filters/FilterConfigIdentifiers.cpp
@@ -10,5 +10,6 @@ const TCHAR* Sequencer_InstancedConfigSubId = TEXT("Outliner");
 const TCHAR* CurveEditor_InstancedConfigId = TEXT("CurveEditor");
 
 const TCHAR* FilterArea_Outliner = TEXT("Outliner");
+const TCHAR* FilterArea_CurveEditor = TEXT("CurveEditor");
 
 } // namespace UE::Sequencer
\ No newline at end of file
diff --git a/Engine/Source/Editor/Sequencer/Private/Filters/FilterConfigIdentifiers.h b/Engine/Source/Editor/Sequencer/Private/Filters/FilterConfigIdentifiers.h
--- a/Engine/Source/Editor/Sequencer/Private/Filters/FilterConfigIdentifiers.h
+++ b/Engine/Source/Editor/Sequencer/Private/Filters/FilterConfigIdentifiers.h
@@ -32,4 +32,9 @@ extern const TCHAR* CurveEditor_InstancedConfigId;
  * @see USequencerSettings::FindOrAddTrackFilterArea 
  */
 extern const TCHAR* FilterArea_Outliner;
+/** 
+ * ID to use for the Curve Editor filter area.
+ * @see USequencerSettings::FindOrAddTrackFilterArea 
+ */
+extern const TCHAR* FilterArea_CurveEditor;
 } // namespace UE::Sequencer
diff --git a/Engine/Source/Editor/Sequencer/Private/Filters/Menus/SequencerViewOptionsMenu.cpp b/Engine/Source/Editor/Sequencer/Private/Filters/Menus/SequencerViewOptionsMenu.cpp
--- a/Engine/Source/Editor/Sequencer/Private/Filters/Menus/SequencerViewOptionsMenu.cpp
+++ b/Engine/Source/Editor/Sequencer/Private/Filters/Menus/SequencerViewOptionsMenu.cpp
@@ -12,33 +12,88 @@
 
 #define LOCTEXT_NAMESPACE "SequencerViewOptionsMenu"
 
-TSharedRef<SWidget> FSequencerViewOptionsMenu::CreateMenu(const TWeakPtr<FSequencer>& InSequencerWeak)
+namespace UE::Sequencer
 {
-	if (!InSequencerWeak.IsValid())
-	{
-		return SNullWidget::NullWidget;
-	}
+namespace ViewOptionsDetail
+{
+static void AddFilterVisibilitySection(UToolMenu& InMenu)
+{
+	const FSequencerTrackFilterCommands& TrackFilterCommands = FSequencerTrackFilterCommands::Get();
+	FToolMenuSection& FilterBarVisibilitySection = InMenu.FindOrAddSection(
+		TEXT("FilterBarVisibility") , LOCTEXT("FilterBarVisibilityHeading", "Filter Bar")
+		);
+	FilterBarVisibilitySection.AddMenuEntry(TrackFilterCommands.ToggleFilterBarVisibility);
+}
 
-	const FName FilterMenuName = TEXT("Sequencer.ViewOptionsMenu");
-	if (!UToolMenus::Get()->IsMenuRegistered(FilterMenuName))
+static TSharedRef<SWidget> CreateMenu(FName InMenuName, FOnPopulateFilterBarMenu InPopulateMenu, const TSharedRef<FUICommandList>& InCommandList)
+{
+	if (!UToolMenus::Get()->IsMenuRegistered(InMenuName))
 	{
-		UToolMenu* const Menu = UToolMenus::Get()->RegisterMenu(FilterMenuName);
+		UToolMenu* const Menu = UToolMenus::Get()->RegisterMenu(InMenuName);
 		Menu->bShouldCloseWindowAfterMenuSelection = false;
-		Menu->AddDynamicSection(NAME_None, FNewToolMenuDelegate::CreateLambda([this](UToolMenu* const InMenu)
-			{
-				if (USequencerMenuContext* const Context = InMenu->FindContext<USequencerMenuContext>())
-				{
-					Context->OnPopulateFilterBarMenu.ExecuteIfBound(InMenu);
-				}
-			}));
+		Menu->AddDynamicSection(NAME_None, FNewToolMenuDelegate::CreateLambda([](UToolMenu* InMenu)
+		{
+		   if (USequencerMenuContext* const Context = InMenu->FindContext<USequencerMenuContext>())
+		   {
+			  Context->OnPopulateFilterBarMenu.ExecuteIfBound(InMenu);
+		   }
+		}));
 	}
+ 
+	USequencerMenuContext* ContextObject = NewObject<USequencerMenuContext>();
+	ContextObject->OnPopulateFilterBarMenu = InPopulateMenu;
+ 
+	const FToolMenuContext MenuContext(InCommandList, nullptr, ContextObject);
+	return UToolMenus::Get()->GenerateWidget(InMenuName, MenuContext);
+}
+}
 
-	USequencerMenuContext* const ContextObject = NewObject<USequencerMenuContext>();
-	ContextObject->Init(InSequencerWeak);
-	ContextObject->OnPopulateFilterBarMenu = FOnPopulateFilterBarMenu::CreateSP(this, &FSequencerViewOptionsMenu::PopulateMenu);
+TSharedRef<SWidget> FSequencerViewOptionsMenu::CreateMenuForSequencer(
+	TWeakPtr<FSequencer> InWeakSequencer, const TSharedRef<FUICommandList>& InCommandList
+	)
+{
+	return ViewOptionsDetail::CreateMenu(
+		TEXT("Sequencer.ViewOptionsMenu"),
+		FOnPopulateFilterBarMenu::CreateSP(this, &FSequencerViewOptionsMenu::PopulateForSequencer, MoveTemp(InWeakSequencer)),
+		InCommandList
+		);
+}
 
-	const FToolMenuContext MenuContext(InSequencerWeak.Pin()->GetCommandBindings(), nullptr, ContextObject);
-	return UToolMenus::Get()->GenerateWidget(FilterMenuName, MenuContext);
+TSharedRef<SWidget> FSequencerViewOptionsMenu::CreateMenuForCurveEditor(
+	TWeakPtr<FSequencer> InWeakSequencer, const TSharedRef<FUICommandList>& InCommandList
+	)
+{
+	return ViewOptionsDetail::CreateMenu(
+		TEXT("Sequencer.ViewOptionsMenu.CurveEditor"),
+		FOnPopulateFilterBarMenu::CreateSP(this, &FSequencerViewOptionsMenu::PopulateForCurveEditor, MoveTemp(InWeakSequencer)),
+		InCommandList
+		);
+}
+
+void FSequencerViewOptionsMenu::PopulateForSequencer(UToolMenu* const InMenu, TWeakPtr<FSequencer> InWeakSequencer)
+{
+	const TSharedPtr<FSequencer> Sequencer = InWeakSequencer.Pin();
+	if (ensure(Sequencer))
+	{
+		WeakSequencer = Sequencer;
+    	
+		PopulateFiltersSection(*InMenu);
+		PopulateSortAndOrganizeSection(*InMenu);
+		PopulateFilterOptionsSection(*InMenu);
+	}
+}
+
+void FSequencerViewOptionsMenu::PopulateForCurveEditor(UToolMenu* const InMenu, TWeakPtr<FSequencer> InWeakSequencer)
+{
+	const TSharedPtr<FSequencer> Sequencer = InWeakSequencer.Pin();
+	if (ensure(Sequencer))
+	{
+		WeakSequencer = Sequencer;
+    	
+		PopulateFiltersSection(*InMenu);
+		// TODO UE-365214: PopulateSortAndOrganizeSection(*InMenu);
+		ViewOptionsDetail::AddFilterVisibilitySection(*InMenu);
+	}
 }
 
 void FSequencerViewOptionsMenu::PopulateMenu(UToolMenu* const InMenu)
@@ -144,8 +199,7 @@ void FSequencerViewOptionsMenu::PopulateFilterOptionsSection(UToolMenu& InMenu)
 		),
 		EUserInterfaceActionType::ToggleButton);
 
-	const FSequencerTrackFilterCommands& TrackFilterCommands = FSequencerTrackFilterCommands::Get();
-	OptionsSection.AddMenuEntry(TrackFilterCommands.ToggleFilterBarVisibility);
+	ViewOptionsDetail::AddFilterVisibilitySection(InMenu);
 }
 
 bool FSequencerViewOptionsMenu::IsIncludePinnedInFilter() const
@@ -265,5 +319,6 @@ TSharedPtr<SSequencer> FSequencerViewOptionsMenu::GetSequencerWidget() const
 	}
 	return nullptr;
 }
+} // namespace UE::Sequencer
 
 #undef LOCTEXT_NAMESPACE
diff --git a/Engine/Source/Editor/Sequencer/Private/Filters/Menus/SequencerViewOptionsMenu.h b/Engine/Source/Editor/Sequencer/Private/Filters/Menus/SequencerViewOptionsMenu.h
--- a/Engine/Source/Editor/Sequencer/Private/Filters/Menus/SequencerViewOptionsMenu.h
+++ b/Engine/Source/Editor/Sequencer/Private/Filters/Menus/SequencerViewOptionsMenu.h
@@ -4,18 +4,30 @@
 
 #include "Templates/SharedPointer.h"
 
+class FUICommandList;
 class FSequencer;
 class SSequencer;
 class SWidget;
 class UToolMenu;
 enum class EFilterBarLayout : uint8;
 
+namespace UE::Sequencer
+{
 class FSequencerViewOptionsMenu : public TSharedFromThis<FSequencerViewOptionsMenu>
 {
 public:
-	TSharedRef<SWidget> CreateMenu(const TWeakPtr<FSequencer>& InSequencerWeak);
+	
+	/** Creates the view options menu for Sequencer's Outliner. */
+	TSharedRef<SWidget> CreateMenuForSequencer(TWeakPtr<FSequencer> InWeakSequencer, const TSharedRef<FUICommandList>& InCommandList);
+	
+	/** Creates the view options menu for Sequencer's Curve Editor. */
+	TSharedRef<SWidget> CreateMenuForCurveEditor(TWeakPtr<FSequencer> InWeakSequencer, const TSharedRef<FUICommandList>& InCommandList);
 
 protected:
+	
+	void PopulateForSequencer(UToolMenu* const InMenu, TWeakPtr<FSequencer> InWeakSequencer);
+	void PopulateForCurveEditor(UToolMenu* const InMenu, TWeakPtr<FSequencer> InWeakSequencer);
+	
 	void PopulateMenu(UToolMenu* const InMenu);
 	void PopulateFiltersSection(UToolMenu& InMenu);
 	void PopulateSortAndOrganizeSection(UToolMenu& InMenu);
@@ -34,3 +46,4 @@ protected:
 
 	TWeakPtr<FSequencer> WeakSequencer;
 };
+} // namespace UE::Sequencer
diff --git a/Engine/Source/Editor/Sequencer/Private/Widgets/Sequencer/SSequencerSearchAndFilterRow.cpp b/Engine/Source/Editor/Sequencer/Private/Widgets/Sequencer/SSequencerSearchAndFilterRow.cpp
--- a/Engine/Source/Editor/Sequencer/Private/Widgets/Sequencer/SSequencerSearchAndFilterRow.cpp
+++ b/Engine/Source/Editor/Sequencer/Private/Widgets/Sequencer/SSequencerSearchAndFilterRow.cpp
@@ -101,7 +101,13 @@ void SSequencerSearchAndFilterRow::Construct(const FArguments& InArgs)
 		.Padding(2.f, 0.f, 0.f, 0.f)
 		[
 			SNew(SSequencerViewOptionsButton)
-			.Sequencer(InArgs._Sequencer)
+			.OnGetMenuContent_Lambda([this, ViewOptions = MakeShared<FSequencerViewOptionsMenu>()]
+			{
+				const TSharedPtr<FSequencer> Sequencer = WeakSequencer.Pin();
+				return Sequencer 
+					? ViewOptions->CreateMenuForSequencer(WeakSequencer, Sequencer->GetCommandBindings().ToSharedRef()) 
+					: SNullWidget::NullWidget;
+			})
 		]
 	];
 }
diff --git a/Engine/Source/Editor/Sequencer/Private/Widgets/Sequencer/SSequencerViewOptionsButton.cpp b/Engine/Source/Editor/Sequencer/Private/Widgets/Sequencer/SSequencerViewOptionsButton.cpp
--- a/Engine/Source/Editor/Sequencer/Private/Widgets/Sequencer/SSequencerViewOptionsButton.cpp
+++ b/Engine/Source/Editor/Sequencer/Private/Widgets/Sequencer/SSequencerViewOptionsButton.cpp
@@ -3,7 +3,7 @@
 #include "SSequencerViewOptionsButton.h"
 
 #include "Filters/Menus/SequencerViewOptionsMenu.h"
-#include "Sequencer.h"
+#include "Widgets/Images/SImage.h"
 #include "Widgets/Input/SComboButton.h"
 
 #define LOCTEXT_NAMESPACE "SSequencerViewOptionsButton"
@@ -14,28 +14,22 @@ namespace UE::Sequencer
 {
 void SSequencerViewOptionsButton::Construct(const FArguments& InArgs)
 {
-	WeakSequencer = InArgs._Sequencer;
-	ViewOptionsMenu = MakeShared<FSequencerViewOptionsMenu>();
-	
 	ChildSlot
 	[
 		SNew(SComboButton)
-			.ContentPadding(2.f)
-			.HAlign(HAlign_Center)
-			.VAlign(VAlign_Center)
-			.ToolTipText(LOCTEXT("ViewOptionsToolTip", "View Options"))
-			.ComboButtonStyle(FAppStyle::Get(), TEXT("SimpleComboButtonWithIcon"))
-			.OnGetMenuContent_Lambda([this]() -> TSharedRef<SWidget>
-			{
-				return ViewOptionsMenu->CreateMenu(WeakSequencer);
-			})
-			.HasDownArrow(false)
-			.ButtonContent()
-			[
-				SNew(SImage)
-				.ColorAndOpacity(FSlateColor::UseForeground())
-				.Image(FAppStyle::Get().GetBrush(TEXT("Icons.Settings")))
-			]
+		.ContentPadding(2.f)
+		.HAlign(HAlign_Center)
+		.VAlign(VAlign_Center)
+		.ToolTipText(LOCTEXT("ViewOptionsToolTip", "View Options"))
+		.ComboButtonStyle(FAppStyle::Get(), TEXT("SimpleComboButtonWithIcon"))
+		.OnGetMenuContent(InArgs._OnGetMenuContent)
+		.HasDownArrow(false)
+		.ButtonContent()
+		[
+			SNew(SImage)
+			.ColorAndOpacity(FSlateColor::UseForeground())
+			.Image(FAppStyle::Get().GetBrush(TEXT("Icons.Settings")))
+		]
 	];
 }
 } // namespace UE::Sequencer
diff --git a/Engine/Source/Editor/Sequencer/Private/Widgets/Sequencer/SSequencerViewOptionsButton.h b/Engine/Source/Editor/Sequencer/Private/Widgets/Sequencer/SSequencerViewOptionsButton.h
--- a/Engine/Source/Editor/Sequencer/Private/Widgets/Sequencer/SSequencerViewOptionsButton.h
+++ b/Engine/Source/Editor/Sequencer/Private/Widgets/Sequencer/SSequencerViewOptionsButton.h
@@ -2,12 +2,10 @@
 
 #pragma once
 
+#include "Framework/SlateDelegates.h"
 #include "Widgets/DeclarativeSyntaxSupport.h"
 #include "Widgets/SCompoundWidget.h"
 
-class FSequencer;
-class FSequencerViewOptionsMenu;
-
 namespace UE::Sequencer
 {
 /** Implements view options for Sequencer. */
@@ -16,18 +14,11 @@ class SSequencerViewOptionsButton : public SCompoundWidget
 public:
 	
 	SLATE_BEGIN_ARGS(SSequencerViewOptionsButton){}
-		/** Used to generate the content of the view options. */
-		SLATE_ARGUMENT(TWeakPtr<FSequencer>, Sequencer)
+		/** Gets the content of the view options. */
+		SLATE_EVENT(FOnGetContent, OnGetMenuContent)
 	SLATE_END_ARGS()
 
 	void Construct(const FArguments& InArgs);
-	
-private:
-	
-	/** Used to implement some items in the view options. */
-	TWeakPtr<FSequencer> WeakSequencer;
-	/** Used to generate the content of the view options. */
-	TSharedPtr<FSequencerViewOptionsMenu> ViewOptionsMenu;
 };
 } // namespace UE::Sequencer
 
