Sequencer: Items displayed Curve Editor now have the same order as the items in Sequencer's Outliner.
- This is achieved by setting a sort delegate in FCurveEditorTree that sorts based on the FViewModel hiearchy.
- This works because the view model hierarchy is sorted in ISequencerSection::GenerateSectionLayout; this is the order in which the Outliner displays items.
- The FCurveEditorTree::SetSortPredicate had to be adjusted in order to extract the ID of the tree item being sorted.
- As consequence, the Niagara API had to be updated as well to fix deprecation warnings.

#Synced-CL 51334859
#Synced-CL 51334859

diff --git a/Engine/Plugins/FX/Niagara/Source/NiagaraEditorWidgets/Private/SNiagaraCurveOverview.cpp b/Engine/Plugins/FX/Niagara/Source/NiagaraEditorWidgets/Private/SNiagaraCurveOverview.cpp
--- a/Engine/Plugins/FX/Niagara/Source/NiagaraEditorWidgets/Private/SNiagaraCurveOverview.cpp
+++ b/Engine/Plugins/FX/Niagara/Source/NiagaraEditorWidgets/Private/SNiagaraCurveOverview.cpp
@@ -329,10 +329,10 @@ struct FNiagaraCurveOverviewTreeItem : public ICurveEditorTreeItem
 		return false;
 	}
 
-	static bool CompareTreeItems(const ICurveEditorTreeItem* ItemA, const ICurveEditorTreeItem* ItemB)
+	static bool CompareTreeItems(const FCurveEditorTreeItem& ItemA, const FCurveEditorTreeItem& ItemB)
 	{
-		const FNiagaraCurveOverviewTreeItem& NiagaraItemA = static_cast<const FNiagaraCurveOverviewTreeItem&>(*ItemA);
-		const FNiagaraCurveOverviewTreeItem& NiagaraItemB = static_cast<const FNiagaraCurveOverviewTreeItem&>(*ItemB);
+		const FNiagaraCurveOverviewTreeItem& NiagaraItemA = static_cast<const FNiagaraCurveOverviewTreeItem&>(*ItemA.GetItem().Get());
+		const FNiagaraCurveOverviewTreeItem& NiagaraItemB = static_cast<const FNiagaraCurveOverviewTreeItem&>(*ItemB.GetItem().Get());
 		return NiagaraItemA.CompareTo(NiagaraItemB);
 	}
 
@@ -477,7 +477,7 @@ void SNiagaraCurveOverview::Construct(const FArguments& InArgs, TSharedRef<FNiag
 	InitParams.ResolveCurveModelDelegate.BindStatic(&UE::NiagaraEditorWidgets::ResolveNiagaraCurveModel); // Needed for undo / redo.
 	CurveEditor->InitCurveEditor(InitParams);
 	CurveEditor->GridLineLabelFormatXAttribute = LOCTEXT("GridXLabelFormat", "{0}");
-	CurveEditor->GetTree()->SetSortPredicate(FCurveEditorTree::FTreeItemSortPredicate::CreateStatic(&FNiagaraCurveOverviewTreeItem::CompareTreeItems));
+	CurveEditor->GetTree()->SetSortPredicate(FCurveEditorTree::FCurveTreeItemSortPredicate::CreateStatic(&FNiagaraCurveOverviewTreeItem::CompareTreeItems));
 
 	CurveEditorPanel = SNew(SCurveEditorPanel, CurveEditor.ToSharedRef())
 		.MinimumViewPanelHeight(0.0f)
diff --git a/Engine/Source/Editor/CurveEditor/Private/Tree/CurveEditorTree.cpp b/Engine/Source/Editor/CurveEditor/Private/Tree/CurveEditorTree.cpp
--- a/Engine/Source/Editor/CurveEditor/Private/Tree/CurveEditorTree.cpp
+++ b/Engine/Source/Editor/CurveEditor/Private/Tree/CurveEditorTree.cpp
@@ -412,9 +412,20 @@ ECurveEditorTreeFilterState FCurveEditorTree::GetFilterState(FCurveEditorTreeIte
 	return FilterStates.Get(InTreeItemID);
 }
 
+PRAGMA_DISABLE_DEPRECATION_WARNINGS
 void FCurveEditorTree::SetSortPredicate(FTreeItemSortPredicate InSortPredicate)
 {
-	SortPredicate = InSortPredicate;
+	SortPredicate = FCurveTreeItemSortPredicate::CreateLambda(
+		[InSortPredicate = MoveTemp(InSortPredicate)](const FCurveEditorTreeItem& ItemA, const FCurveEditorTreeItem& ItemB)
+		{
+			return InSortPredicate.Execute(ItemA.GetItem().Get(), ItemB.GetItem().Get());
+		});
+}
+PRAGMA_ENABLE_DEPRECATION_WARNINGS
+
+void FCurveEditorTree::SetSortPredicate(FCurveTreeItemSortPredicate InSortPredicate)
+{
+	SortPredicate = MoveTemp(InSortPredicate);
 }
 
 void FCurveEditorTree::SortTreeItems()
@@ -437,7 +448,10 @@ void FCurveEditorTree::SortTreeItems(FSortedCurveEditorTreeItems& TreeItemIDsToS
 	// If there is more than one item, sort the items and then repopulate the ChildIDs from the sorted list of items.
 	if (TreeItemIDsToSort.bRequiresSort && TreeItemsToSort.Num() > 1)
 	{
-		TreeItemsToSort.Sort([this](const FCurveEditorTreeItem& ItemA, const FCurveEditorTreeItem& ItemB) { return SortPredicate.Execute(ItemA.GetItem().Get(), ItemB.GetItem().Get()); });
+		TreeItemsToSort.Sort([this](const FCurveEditorTreeItem& ItemA, const FCurveEditorTreeItem& ItemB)
+		{
+			return SortPredicate.Execute(ItemA, ItemB);
+		});
 		for (int32 i = 0; i < TreeItemsToSort.Num(); ++i)
 		{
 			TreeItemIDsToSort.ChildIDs[i] = TreeItemsToSort[i]->GetID();
diff --git a/Engine/Source/Editor/CurveEditor/Public/Tree/CurveEditorTree.h b/Engine/Source/Editor/CurveEditor/Public/Tree/CurveEditorTree.h
--- a/Engine/Source/Editor/CurveEditor/Public/Tree/CurveEditorTree.h
+++ b/Engine/Source/Editor/CurveEditor/Public/Tree/CurveEditorTree.h
@@ -364,8 +364,12 @@ DECLARE_MULTICAST_DELEGATE_OneParam(FOnCurveEditorToggleExpansionState, bool);
 class FCurveEditorTree
 {
 public:
+	
 	/** Defines a predicate for sorting curve editor tree item implementations. */
+	UE_DEPRECATED(5.8, "Use FCurveTreeItemSortPredicate instead.")
 	DECLARE_DELEGATE_RetVal_TwoParams(bool, FTreeItemSortPredicate, const ICurveEditorTreeItem* /* ItemA */, const ICurveEditorTreeItem* /* ItemB */);
+	/** Returns whether ItemA < ItemB. */
+	DECLARE_DELEGATE_RetVal_TwoParams(bool, FCurveTreeItemSortPredicate, const FCurveEditorTreeItem& /* ItemA */, const FCurveEditorTreeItem& /* ItemB */);
 
 	/** Structure containing all the events for this tree */
 	FCurveEditorTreeEvents Events;
@@ -457,7 +461,12 @@ public:
 	/**
 	 * Sets a predicate which will be used to sort tree items after they're been marked as needing sort.
 	 */
-	 UE_API void SetSortPredicate(FTreeItemSortPredicate InSortPredicate);
+	PRAGMA_DISABLE_DEPRECATION_WARNINGS
+	UE_API void SetSortPredicate(FTreeItemSortPredicate InSortPredicate);
+	PRAGMA_ENABLE_DEPRECATION_WARNINGS
+	
+	/** Sets a predicate which will be used to sort tree items after they're been marked as needing sort. */
+	UE_API void SetSortPredicate(FCurveTreeItemSortPredicate InSortPredicate);
 
 	 /**
 	  * Sorts all tree items which have been marked for sorting if the sort predicate has been set.
@@ -582,7 +591,7 @@ private:
 	FCurveEditorFilterStates FilterStates;
 
 	/** A predicate which will be used to sort tree items after they're been marked as needing sort. */
-	FTreeItemSortPredicate SortPredicate;
+	FCurveTreeItemSortPredicate SortPredicate;
 
 	/** Delegate for when toggle expansion state is invoked */
 	FOnCurveEditorToggleExpansionState ToggleExpansionStateDelegate;
diff --git a/Engine/Source/Editor/Sequencer/Private/CurveEditor/SequencerCurveEditorApp.cpp b/Engine/Source/Editor/Sequencer/Private/CurveEditor/SequencerCurveEditorApp.cpp
--- a/Engine/Source/Editor/Sequencer/Private/CurveEditor/SequencerCurveEditorApp.cpp
+++ b/Engine/Source/Editor/Sequencer/Private/CurveEditor/SequencerCurveEditorApp.cpp
@@ -10,6 +10,7 @@
 #include "ICurveEditorModule.h"
 #include "MVVM/CurveEditorExtension.h"
 #include "MVVM/ViewModels/SequencerEditorViewModelImpl.h"
+#include "MVVM/ViewModels/SequencerModelUtils.h"
 #include "Specializations/SequencerCurveEditor.h"
 #include "Specializations/SequencerCurveEditorBounds.h"
 #include "Specializations/SequencerCurveEditorToolbarExtender.h"
@@ -100,6 +101,19 @@ FSequencerCurveEditorApp::FSequencerCurveEditorApp(const TSharedRef<FSequencer>&
 	{
 		WidgetOwner.SyncSelection();
 	});
+	
+	// The items in curve editor are supposed to have the same order as in Sequencer's Outliner.
+	CurveEditor->GetTree()->SetSortPredicate(FCurveEditorTree::FCurveTreeItemSortPredicate::CreateLambda(
+		[this](const FCurveEditorTreeItem& ItemA, const FCurveEditorTreeItem& ItemB)
+		{
+			FCurveModelSyncer& Syncer = SelectionSyncLogic.GetSequencerToCurveEditorSyncer();
+			const TSharedPtr<FViewModel> ViewModelA = Syncer.FindViewModelById(ItemA.GetID()).Pin();
+			const TSharedPtr<FViewModel> ViewModelB = Syncer.FindViewModelById(ItemB.GetID()).Pin();
+			return ensure(ViewModelA && ViewModelB)
+				? ComesFirstInHierarchy(FViewModelPtr(ViewModelA.Get()), FViewModelPtr(ViewModelB.Get()))
+				// Sorting algorithms require consistent < implementation so fallback to pointer comparision.
+				: ViewModelA.Get() < ViewModelB.Get();
+		}));
 }
 
 FSequencerCurveEditorApp* FSequencerCurveEditorApp::Get(const ISequencer& InSequencer)
diff --git a/Engine/Source/Editor/Sequencer/Private/MVVM/ViewModels/SequencerModelUtils.cpp b/Engine/Source/Editor/Sequencer/Private/MVVM/ViewModels/SequencerModelUtils.cpp
--- a/Engine/Source/Editor/Sequencer/Private/MVVM/ViewModels/SequencerModelUtils.cpp
+++ b/Engine/Source/Editor/Sequencer/Private/MVVM/ViewModels/SequencerModelUtils.cpp
@@ -39,5 +39,81 @@ TViewModelPtr<ITrackExtension> GetParentTrackNodeAndNamePath(const TViewModelPtr
 	return nullptr;
 }
 
+namespace SharedParentDetail
+{
+static TArray<TViewModelPtr<FViewModel>> GetAncestorsIncludingThis(const FViewModel& InViewModel)
+{
+	TArray<TViewModelPtr<FViewModel>> Ancestors;
+	for (const TViewModelPtr<FViewModel>& Item : InViewModel.GetAncestors())
+	{
+		Ancestors.Add(Item);
+	}
+	return Ancestors;
+}
+}
+
+TOptional<FSharedParentInfo> FindSharedParent(FViewModel& ViewModelA, FViewModel& ViewModelB)
+{
+	// Most common fast-path. Avoids TArray allocations.
+	if (ViewModelA.GetParent() == ViewModelB.GetParent())
+	{
+		return FSharedParentInfo{ ViewModelA.GetParent(), FViewModelPtr(&ViewModelA), FViewModelPtr(&ViewModelB) };
+	}
+	
+	const TArray<TViewModelPtr<FViewModel>> AncestorsA = SharedParentDetail::GetAncestorsIncludingThis(ViewModelA);
+	const TArray<TViewModelPtr<FViewModel>> AncestorsB = SharedParentDetail::GetAncestorsIncludingThis(ViewModelB);
+	const int32 MinNum = FMath::Min(AncestorsA.Num(), AncestorsB.Num());
+	for (int32 Index = 0; Index < MinNum; ++Index)
+	{
+		const bool bModelsShareParent = AncestorsA[Index] == AncestorsB[Index];
+		if (bModelsShareParent)
+		{
+			continue;
+		}
+		
+		if (Index == 0)
+		{
+			// The view models don't share any parent. 
+			return {};
+		}
+		
+		const TViewModelPtr<FViewModel> LastSharedParent = AncestorsA[Index - 1];
+		return FSharedParentInfo(LastSharedParent, AncestorsA[Index], AncestorsB[Index]);
+	}
+	
+	return {};
+}
+
+bool ComesFirstInHierarchy(const FViewModelPtr& ViewModelA, const FViewModelPtr& ViewModelB)
+{
+	const TOptional<FSharedParentInfo> ParentInfo = ViewModelA && ViewModelB 
+		? FindSharedParent(*ViewModelA, *ViewModelB) 
+		: TOptional<FSharedParentInfo>{};
+	
+	// Array sorting algorithms require consistent <. We default to this in error cases.
+	const bool bIsLessByPtr = ViewModelA.Get() < ViewModelB.Get(); 
+	if (!ensure(ParentInfo))
+	{
+		return bIsLessByPtr;
+	}
+	
+	const FViewModelPtr& Parent = ParentInfo->Parent;
+	for (const FViewModelPtr& Child : Parent->GetChildren())
+	{
+		if (Child == ViewModelA)
+		{
+			// ViewModelA appeared before ViewModelB, so ViewmodelA < ViewModelB
+			return true;
+		}
+		
+		if (Child == ViewModelB)
+		{
+			// ViewModelB appeared before ViewModelA, so ViewModelA > ViewModelB.
+			return false;
+		}
+	}
+				
+	return bIsLessByPtr;
+}
 } // namespace Sequencer
 } // namespace UE
diff --git a/Engine/Source/Editor/Sequencer/Private/MVVM/ViewModels/SequencerModelUtils.h b/Engine/Source/Editor/Sequencer/Private/MVVM/ViewModels/SequencerModelUtils.h
--- a/Engine/Source/Editor/Sequencer/Private/MVVM/ViewModels/SequencerModelUtils.h
+++ b/Engine/Source/Editor/Sequencer/Private/MVVM/ViewModels/SequencerModelUtils.h
@@ -22,5 +22,29 @@ class ITrackExtension;
  */
 TViewModelPtr<ITrackExtension> GetParentTrackNodeAndNamePath(const TViewModelPtr<IOutlinerExtension>& Node, TArray<FName>& OutNamePath);
 
+/** 
+ * Information about two FViewModels ItemA and ItemB that share some parent.
+ * Examples: 
+ *	- ItemA = Location.X, ItemB = Location.Y		->		 Parent = Location, DirectChildA = Location.X, DirectChildB = Location.Y
+ *	- ItemA = Location.X, ItemB = Rotation.Y		->		 Parent = Transform, DirectChildA = Location, DirectChildB = Rotation
+ */
+struct FSharedParentInfo
+{
+	/** The parent that both ItemA and ItemB share. */
+	FViewModelPtr Parent;
+	/** Immediate child of Parent that leads to ItemA. If ItemA->GetParent() == Parent, then DirectChildA == ItemA. */
+	FViewModelPtr DirectChildA;
+	/** Immediate child of Parent that leads to ItemB. If ItemB->GetParent() == Parent, then DirectChildB == ItemB. */
+	FViewModelPtr DirectChildB;
+};
+
+/** Walks up the hierarchy of both view models and finds the first shared parent. Returns unset if there is no shared parent. */
+TOptional<FSharedParentInfo> FindSharedParent(const FViewModel& ViewModelA, const FViewModel& ViewModelB);
+
+/**
+ * Utility for implementing ViewModelA < ViewModelB based on whether ViewModelA appears first in the tree hierarchy.
+ * @return Whether ViewModelA comes first in the tree hierarchy.
+ */
+bool ComesFirstInHierarchy(const FViewModelPtr& ViewModelA, const FViewModelPtr& ViewModelB);
 } // namespace Sequencer
 } // namespace UE
\ No newline at end of file
