Sequencer: Adjust padding in Curve Editor filter area to match that in Sequencer.
Also add a SScrollBox so the search row shows a shadow to the right when too narrow... matching behaviour of Sequencer's Outliner.

#Synced-CL 51334859

diff --git a/Engine/Source/Editor/EditorStyle/Private/StarshipStyle.cpp b/Engine/Source/Editor/EditorStyle/Private/StarshipStyle.cpp
--- a/Engine/Source/Editor/EditorStyle/Private/StarshipStyle.cpp
+++ b/Engine/Source/Editor/EditorStyle/Private/StarshipStyle.cpp
@@ -2536,6 +2536,9 @@ void FStarshipEditorStyle::FStyle::SetupSequencerStyles()
 		Set("Sequencer.TrackIsolate", new IMAGE_BRUSH_SVG("Sequencer/TrackIsolate", Icon16x16));
 		Set("Sequencer.TrackHide", new IMAGE_BRUSH_SVG("Sequencer/TrackHide", Icon16x16));
 		Set("Sequencer.TrackShow", new IMAGE_BRUSH_SVG("Sequencer/TrackShow", Icon16x16));
+		
+		constexpr float CommonPadding = 3;
+		Set("Sequencer.FilterAreaMargin", FMargin(2 * CommonPadding, CommonPadding));
 	}
 
 
diff --git a/Engine/Source/Editor/Sequencer/Private/CurveEditor/Widgets/Tree/SCurveTreeContent.cpp b/Engine/Source/Editor/Sequencer/Private/CurveEditor/Widgets/Tree/SCurveTreeContent.cpp
--- a/Engine/Source/Editor/Sequencer/Private/CurveEditor/Widgets/Tree/SCurveTreeContent.cpp
+++ b/Engine/Source/Editor/Sequencer/Private/CurveEditor/Widgets/Tree/SCurveTreeContent.cpp
@@ -75,29 +75,41 @@ void SCurveTreeContent::SyncSelection() const
 
 TSharedRef<SWidget> SCurveTreeContent::MakeFilterArea(const FArguments& InArgs)
 {
-	return SNew(SVerticalBox)
-		
-		+SVerticalBox::Slot()
-		.AutoHeight()
+	return SNew(SBorder)
+		.BorderImage(FAppStyle::GetBrush(TEXT("ToolPanel.GroupBorder")))
+		.Padding(FAppStyle::GetMargin(TEXT("Sequencer.FilterAreaMargin")))
+		.Clipping(EWidgetClipping::ClipToBounds)
 		[
-			SNew(SBorder)
-			.BorderImage(FAppStyle::GetBrush("ToolPanel.GroupBorder"))
-			.Clipping(EWidgetClipping::ClipToBounds)
+			SNew(SVerticalBox)
+		
+			+SVerticalBox::Slot()
+			.AutoHeight()
 			[
-				SAssignNew(FilterRow, SCurveEditorSearchAndFilterRow)
-				.Sequencer(InArgs._Sequencer)
-				.FilterViewModel(InArgs._FilterViewModel)
-				.CommandList(InArgs._CommandList)
-				.ScrollItemIntoView(this, &SCurveTreeContent::ScrollIntoView)
-				.AllItems(this, &SCurveTreeContent::GetAllItems)
-				.SelectedItems(this, &SCurveTreeContent::GetSelectedItems)
+				// The scroll box adds a shadow to the right when the toolbar is narrow and hiding buttons.
+				SNew(SScrollBox)
+				.Orientation(Orient_Horizontal)
+				.ScrollBarAlwaysVisible(false)
+				.ScrollBarVisibility(EVisibility::Collapsed)
+				
+				+ SScrollBox::Slot()
+				.HAlign(HAlign_Fill)
+				.FillContentSize(1.f)
+				[
+					SAssignNew(FilterRow, SCurveEditorSearchAndFilterRow)
+					.Sequencer(InArgs._Sequencer)
+					.FilterViewModel(InArgs._FilterViewModel)
+					.CommandList(InArgs._CommandList)
+					.ScrollItemIntoView(this, &SCurveTreeContent::ScrollIntoView)
+					.AllItems(this, &SCurveTreeContent::GetAllItems)
+					.SelectedItems(this, &SCurveTreeContent::GetSelectedItems)
+				]
+			]
+			
+			+SVerticalBox::Slot()
+			.AutoHeight()
+			[
+				InArgs._FilterPills.Widget
 			]
-		]
-		
-		+SVerticalBox::Slot()
-		.AutoHeight()
-		[
-			InArgs._FilterPills.Widget
 		];
 }
 
diff --git a/Engine/Source/Editor/Sequencer/Private/SSequencer.cpp b/Engine/Source/Editor/Sequencer/Private/SSequencer.cpp
--- a/Engine/Source/Editor/Sequencer/Private/SSequencer.cpp
+++ b/Engine/Source/Editor/Sequencer/Private/SSequencer.cpp
@@ -675,7 +675,7 @@ TSharedRef<SSplitter> SSequencer::ConstructOutlinerAndTrackArea(const FArguments
 				[
 					SNew(SBorder)
 					.BorderImage(FAppStyle::GetBrush(TEXT("ToolPanel.GroupBorder")))
-					.Padding(FMargin(CommonPadding * 2, CommonPadding))
+					.Padding(FAppStyle::GetMargin(TEXT("Sequencer.FilterAreaMargin")))
 					.Clipping(EWidgetClipping::ClipToBounds)
 					[
 						SAssignNew(SearchAndFilterRow, SVerticalBox)
