Sequencer: Remove "Link Filters with Curve Editor" view option entry when "Sequencer.CurveEditor.UseLinkedFiltering" CVar is enabled to simulate what the feature would look like. We'll still have to remove the underlying setting at a later point.

#Synced-CL 51334859

diff --git a/Engine/Source/Editor/Sequencer/Private/Filters/Menus/SequencerViewOptionsMenu.cpp b/Engine/Source/Editor/Sequencer/Private/Filters/Menus/SequencerViewOptionsMenu.cpp
--- a/Engine/Source/Editor/Sequencer/Private/Filters/Menus/SequencerViewOptionsMenu.cpp
+++ b/Engine/Source/Editor/Sequencer/Private/Filters/Menus/SequencerViewOptionsMenu.cpp
@@ -1,12 +1,14 @@
 // Copyright Epic Games, Inc. All Rights Reserved.
 
 #include "SequencerViewOptionsMenu.h"
+
+#include "CurveEditor/Views/CurveEditorWidgetOwner.h"
 #include "Filters/Menus/SequencerMenuContext.h"
 #include "Filters/SequencerTrackFilterCommands.h"
+#include "SSequencer.h"
 #include "Sequencer.h"
 #include "SequencerCommands.h"
 #include "SequencerFilterBarContext.h"
-#include "SSequencer.h"
 #include "ToolMenu.h"
 #include "ToolMenus.h"
 
@@ -188,16 +190,20 @@ void FSequencerViewOptionsMenu::PopulateFilterOptionsSection(UToolMenu& InMenu)
 		),
 		EUserInterfaceActionType::ToggleButton);
 
-	OptionsSection.AddMenuEntry(TEXT("LinkFiltersWithCurveEditor"),
-		LOCTEXT("LinkFiltersWithCurveEditor", "Link Filters with Curve Editor"),
-		LOCTEXT("LinkFiltersWithCurveEditorToolTip", "Toggle linking filters with the curve editor tree"),
-		FSlateIcon(),
-		FUIAction(
-			FExecuteAction::CreateRaw(this, &FSequencerViewOptionsMenu::ToggleLinkFiltersWithCurveEditor),
-			FCanExecuteAction(),
-			FIsActionChecked::CreateRaw(this, &FSequencerViewOptionsMenu::IsLinkFiltersWithCurveEditor)
-		),
-		EUserInterfaceActionType::ToggleButton);
+	// TODO UE-364008: Remove this once we decide to permanently support linked filtering in Curve Editor
+	if (!UseLinkedFilteringUI())
+	{
+		OptionsSection.AddMenuEntry(TEXT("LinkFiltersWithCurveEditor"),
+			LOCTEXT("LinkFiltersWithCurveEditor", "Link Filters with Curve Editor"),
+			LOCTEXT("LinkFiltersWithCurveEditorToolTip", "Toggle linking filters with the curve editor tree"),
+			FSlateIcon(),
+			FUIAction(
+				FExecuteAction::CreateRaw(this, &FSequencerViewOptionsMenu::ToggleLinkFiltersWithCurveEditor),
+				FCanExecuteAction(),
+				FIsActionChecked::CreateRaw(this, &FSequencerViewOptionsMenu::IsLinkFiltersWithCurveEditor)
+			),
+			EUserInterfaceActionType::ToggleButton);
+	}
 
 	ViewOptionsDetail::AddFilterVisibilitySection(InMenu);
 }
