Fix for subsequnce parent ranges not being drawn when focusing them from their parent if the range was outside the subsequence asset's start and end frames.

The time warp refactor introduced a change that clamped the subsequence range used for display to the intersection of the asset's start and end frames. This change no longer clamps this value, and will always draw both the asset range, and the range the parent sequence evaluates.

#Synced-CL 51549532

diff --git a/Engine/Source/Editor/Sequencer/Private/Sequencer.cpp b/Engine/Source/Editor/Sequencer/Private/Sequencer.cpp
--- a/Engine/Source/Editor/Sequencer/Private/Sequencer.cpp
+++ b/Engine/Source/Editor/Sequencer/Private/Sequencer.cpp
@@ -1765,8 +1765,7 @@ void FSequencer::UpdateSubSequenceData()
 
 				CurrentTime = SubData->OuterToInnerTransform.TransformTime(CurrentTime);
 
-				// Clamp the range and transform it into inner space
-				ValidRange = TRange<FFrameTime>::Intersection(ValidRange, ConvertToFrameTimeRange(SubData->ParentPlayRange.Value));
+				ValidRange = ConvertToFrameTimeRange(SubData->ParentPlayRange.Value);
 				if (ValidRange.IsEmpty())
 				{
 					break;
