Sequencer: Fix undo/redo for anim mixer layer reordering

FTrackModel::PerformDrop() calls ExecuteReorder() without wrapping it in
an FScopedTransaction. Each SetRowIndex() call inside creates separate
undo entries via Modify(), so undo steps through intermediate reorder
states instead of reverting the entire reorder atomically.

#Synced-CL 51540725

diff --git a/Engine/Source/Editor/Sequencer/Private/MVVM/ViewModels/TrackModel.cpp b/Engine/Source/Editor/Sequencer/Private/MVVM/ViewModels/TrackModel.cpp
--- a/Engine/Source/Editor/Sequencer/Private/MVVM/ViewModels/TrackModel.cpp
+++ b/Engine/Source/Editor/Sequencer/Private/MVVM/ViewModels/TrackModel.cpp
@@ -790,6 +790,7 @@ void FTrackModel::PerformDrop(const FViewModelPtr& TargetModel, const FDragDropE
 					TargetRowIndex++;
 				}
 
+				const FScopedTransaction Transaction(NSLOCTEXT("SequencerTrackNode", "ReorderTrackRows", "Reorder Track Rows"));
 				TrackRowReorderOp->ExecuteReorder(TargetRowIndex);
 			}
 		}
