Sequencer: Wrap track toggle actions in transactions for undo

ToggleTrackSelectionLock, ToggleTrackSelectionDeactive,
ToggleTrackSelectionMute, and ToggleTrackSelectionSolo all call
Modify() on multiple objects without an FScopedTransaction. Each
Modify() creates a separate undo entry, so undoing piecemeal can
leave Control Rig components in an inconsistent attach state,
triggering an ensure in SceneComponent::DetachFromComponent.

#Synced-CL 51592964

diff --git a/Engine/Source/Editor/Sequencer/Private/Sequencer.cpp b/Engine/Source/Editor/Sequencer/Private/Sequencer.cpp
--- a/Engine/Source/Editor/Sequencer/Private/Sequencer.cpp
+++ b/Engine/Source/Editor/Sequencer/Private/Sequencer.cpp
@@ -13171,6 +13171,8 @@ void FSequencer::ToggleTrackSelectionLock()
 
 	const bool bNewState = !bAllSame;
 
+	const FScopedTransaction Transaction(NSLOCTEXT("Sequencer", "ToggleLock", "Toggle Lock"));
+
 	bool bAnyChanged = false;
 
 	for (const TViewModelPtr<IOutlinerExtension>& OutlinerItem : OutlinerSelection)
@@ -13212,6 +13214,8 @@ void FSequencer::ToggleTrackSelectionDeactive()
 
 	const bool bNewState = !bAllSame;
 
+	const FScopedTransaction Transaction(NSLOCTEXT("Sequencer", "ToggleDeactivate", "Toggle Deactivate"));
+
 	bool bAnyChanged = false;
 
 	for (const TViewModelPtr<IOutlinerExtension>& OutlinerItem : OutlinerSelection)
@@ -13252,6 +13256,8 @@ void FSequencer::ToggleTrackSelectionMute()
 
 	const bool bNewState = !bAllSame;
 
+	const FScopedTransaction Transaction(NSLOCTEXT("Sequencer", "ToggleMute", "Toggle Mute"));
+
 	bool bAnyChanged = false;
 
 	for (const TViewModelPtr<IOutlinerExtension>& OutlinerItem : OutlinerSelection)
@@ -13292,6 +13298,8 @@ void FSequencer::ToggleTrackSelectionSolo()
 
 	const bool bNewState = !bAllSame;
 
+	const FScopedTransaction Transaction(NSLOCTEXT("Sequencer", "ToggleSolo", "Toggle Solo"));
+
 	bool bAnyChanged = false;
 
 	for (const TViewModelPtr<IOutlinerExtension>& OutlinerItem : OutlinerSelection)
