Sequencer: Use TimeBounds when playing back in editor.

Time bounds will return the selection range if it exists, otherwise RootTimeBounds. This matches the behavior of rewind, jumping from key to key, etc.

#Synced-CL 51563788

#Synced-CL 51563788

#Synced-CL 51563788

diff --git a/Engine/Source/Editor/Sequencer/Private/Sequencer.cpp b/Engine/Source/Editor/Sequencer/Private/Sequencer.cpp
--- a/Engine/Source/Editor/Sequencer/Private/Sequencer.cpp
+++ b/Engine/Source/Editor/Sequencer/Private/Sequencer.cpp
@@ -5201,10 +5201,10 @@ FReply FSequencer::OnPlay(bool bTogglePlay)
 	}
 	else
 	{
-		TRange<FFrameNumber> RootTimeBounds = GetRootTimeBounds();
+		TRange<FFrameNumber> TimeBounds = GetTimeBounds();
 
-		FFrameNumber MinInclusiveTime = UE::MovieScene::DiscreteInclusiveLower(RootTimeBounds);
-		FFrameNumber MaxInclusiveTime = UE::MovieScene::DiscreteExclusiveUpper(RootTimeBounds) - 1;
+		FFrameNumber MinInclusiveTime = UE::MovieScene::DiscreteInclusiveLower(TimeBounds);
+		FFrameNumber MaxInclusiveTime = UE::MovieScene::DiscreteExclusiveUpper(TimeBounds) - 1;
 		
 		FFrameTime GlobalTime = GlobalPlaybackWarpTransform.TransformTime(GetGlobalTime().Time);
 
