Sequencer: Fix filter pills being cut off in Curve Editor when width is insufficient.

There was a calculation error in SFilterBarClippingHorizontalBox when checking whether pills overlap the wrap button.
When the pill goes outside of the parent size, the pill should be clipped. Example: ChildRightX = 1500 and ParentSizeX = 500 should be clipped. However, the calculation was wrong because ChildRightX was using absolute position and ParentSizeX was using local position.
This bug never surfaced because Sequencer is usually aligned to the left-most part of the UI, in which case the absolute position X is nearly 0 as well.

As result, the ParentReduction size also had to be adjusted. With the old values, the combo button would overlap the pills.

#Synced-CL 51334859

diff --git a/Engine/Source/Editor/Sequencer/Private/Filters/Widgets/SFilterBarClippingHorizontalBox.cpp b/Engine/Source/Editor/Sequencer/Private/Filters/Widgets/SFilterBarClippingHorizontalBox.cpp
--- a/Engine/Source/Editor/Sequencer/Private/Filters/Widgets/SFilterBarClippingHorizontalBox.cpp
+++ b/Engine/Source/Editor/Sequencer/Private/Filters/Widgets/SFilterBarClippingHorizontalBox.cpp
@@ -31,7 +31,7 @@ void SFilterBarClippingHorizontalBox::OnArrangeChildren(const FGeometry& InAllot
 	}
 
 	const TSharedPtr<SComboButton> WrapButton = WeakWrapButton.Pin();
-	constexpr int32 ParentReductionSize = 20; // Arbitrary size to make things look nice
+	constexpr int32 ParentReductionSize = 10; // Arbitrary size to make things look nice
 
 	LastClippedIndex = ClippedIndex;
 
@@ -78,9 +78,9 @@ void SFilterBarClippingHorizontalBox::OnArrangeChildren(const FGeometry& InAllot
 			const FArrangedWidget& ChildWidget = ArrangedChildren[ChildIndex];
 
 			const int32 ChildRightX = FMath::TruncToInt(ChildWidget.Geometry.AbsolutePosition.X + ChildWidget.Geometry.GetLocalSize().X * ChildWidget.Geometry.Scale);
-			const int32 ParentSizeX = FMath::TruncToInt(InAllottedGeometry.GetLocalSize().X * InAllottedGeometry.Scale);
+			const int32 ParentSizeX = FMath::TruncToInt(InAllottedGeometry.GetAbsolutePosition().X + InAllottedGeometry.GetLocalSize().X * InAllottedGeometry.Scale);
 
-			if (ChildRightX > ParentSizeX + ParentReductionSize)
+			if (ChildRightX > ParentSizeX - ParentReductionSize)
 			{
 				++NumClippedChildren;
 				ArrangedChildren.Remove(ChildIndex);
