Sequencer: When selecting items in Curve Editor, the equivalent item will only be expanded in Sequencer's Outliner if "Auto Expand Nodes on Selection" is enabled.

#Synced-CL 51334859

diff --git a/Engine/Source/Editor/Sequencer/Private/CurveEditor/Integration/Selection/SequencerSelectionSyncer.cpp b/Engine/Source/Editor/Sequencer/Private/CurveEditor/Integration/Selection/SequencerSelectionSyncer.cpp
--- a/Engine/Source/Editor/Sequencer/Private/CurveEditor/Integration/Selection/SequencerSelectionSyncer.cpp
+++ b/Engine/Source/Editor/Sequencer/Private/CurveEditor/Integration/Selection/SequencerSelectionSyncer.cpp
@@ -34,6 +34,8 @@ FSequencerSelectionSyncer::~FSequencerSelectionSyncer()
 	}
 }
 
+static TAutoConsoleVariable<bool> CVarSelectionExpandsSequencer(TEXT("Sequencer.CurveEditor.SelectionExpandsSequencer"), false, TEXT(""));
+
 void FSequencerSelectionSyncer::OnCurveEditorSelectionChanged()
 {
 	const TSharedPtr<FSequencer> Sequencer = WeakSequencer.Pin();
@@ -52,10 +54,15 @@ void FSequencerSelectionSyncer::OnCurveEditorSelectionChanged()
 	
 	FSequencerSelection& TargetSelection = Sequencer->GetSelection();
 	ApplyCurveEditorToSequencer(TargetSelection, *Tree, CurveModelSyncer);
-	ExpandSelectedSequencerItems(TargetSelection);
 	
-	// This is needed to apply the expansion state by causing SOutlinerView::Refresh.
-	Sequencer->RefreshTree();
+	// E.g. in Curve Editor click Cube.Transform.Location.X but Cube is collapsed in Sequencer, so the user wouldn't see its now selected in Sequencer.
+	// So we expand to that item in Sequencer's Outliner if configured to do so by the user.
+	if (Settings->GetAutoExpandNodesOnSelection())
+	{
+		ExpandSelectedSequencerItems(TargetSelection);
+		// This is needed to apply the expansion state by causing SOutlinerView::Refresh.
+		Sequencer->RefreshTree();
+	}
 	
 	// If the user clicked a single item, show it in the UI
 	const TSet<TWeakViewModelPtr<IOutlinerExtension>>& Selected = TargetSelection.Outliner.GetSelected();
