Masking In Anim Mixer - UI
-Anim Mixer Track model now only supports dragging and dropping of Skeletal Animation sections
-Added SectionProviderOutlinerModel to support displaying section provider decorations, this allows the decoration outliner to group sections from the same decoration on the same row, and support mutliple sections on the same decoration
-The Decoration outliner model now uses the section provider outliner instead of the Section Outliner to hold sections.
-Added Mask decoration editor to support creating mask sections in the decoration outliner
-SequencerHelpers::BuildDecorationMenu now takes a TWeakPtr<ISequencer> to support getting the current time when adding new sections to existing decorations that support them.

#Synced-CL 51628111

diff --git a/Engine/Plugins/Experimental/MovieSceneAnimMixer/Source/MovieSceneAnimMixer/Private/MovieSceneAnimMixerModule.cpp b/Engine/Plugins/Experimental/MovieSceneAnimMixer/Source/MovieSceneAnimMixer/Private/MovieSceneAnimMixerModule.cpp
--- a/Engine/Plugins/Experimental/MovieSceneAnimMixer/Source/MovieSceneAnimMixer/Private/MovieSceneAnimMixerModule.cpp
+++ b/Engine/Plugins/Experimental/MovieSceneAnimMixer/Source/MovieSceneAnimMixer/Private/MovieSceneAnimMixerModule.cpp
@@ -8,6 +8,7 @@
 #include "MovieSceneAnimationMixerLayer.h"
 #include "MovieSceneAnimationMixerTrack.h"
 #include "MovieSceneRootMotionTargetDecoration.h"
+#include "MovieSceneAnimationMaskDecoration.h"
 #include "Modules/ModuleInterface.h"
 #include "Modules/ModuleManager.h"
 #include "Sections/MovieSceneSkeletalAnimationSection.h"
@@ -51,6 +52,8 @@ void FMovieSceneAnimMixerModule::StartupModule()
 	// Register Root Motion Settings decoration as compatible with Skeletal Animation Sections
 	IMovieSceneModule::Get().RegisterCompatibleDecoration(UMovieSceneSkeletalAnimationSection::StaticClass(), UMovieSceneRootMotionSettingsDecoration::StaticClass());
 
+	IMovieSceneModule::Get().RegisterCompatibleDecoration(UMovieSceneAnimationMixerLayer::StaticClass(), UMovieSceneAnimationMaskDecoration::StaticClass());
+
 	IModularFeatures::Get().RegisterModularFeature(IAnimGraphRuntime_SequencerMixerTargetConnector::GetModularFeatureName(), this);
 }
 
@@ -58,6 +61,7 @@ void FMovieSceneAnimMixerModule::ShutdownModule()
 {
 	IMovieSceneModule::Get().UnregisterCompatibleDecoration(UMovieSceneAnimationMixerTrack::StaticClass(), UMovieSceneRootMotionTargetDecoration::StaticClass());
 	IMovieSceneModule::Get().UnregisterCompatibleDecoration(UMovieSceneSkeletalAnimationSection::StaticClass(), UMovieSceneRootMotionSettingsDecoration::StaticClass());
+	IMovieSceneModule::Get().UnregisterCompatibleDecoration(UMovieSceneAnimationMixerLayer::StaticClass(), UMovieSceneAnimationMaskDecoration::StaticClass());
 
 	IModularFeatures::Get().UnregisterModularFeature(IAnimGraphRuntime_SequencerMixerTargetConnector::GetModularFeatureName(), this);
 }
diff --git a/Engine/Plugins/Experimental/MovieSceneAnimMixer/Source/MovieSceneAnimMixer/Public/MovieSceneAnimTransitionSectionBase.h b/Engine/Plugins/Experimental/MovieSceneAnimMixer/Source/MovieSceneAnimMixer/Public/MovieSceneAnimTransitionSectionBase.h
--- a/Engine/Plugins/Experimental/MovieSceneAnimMixer/Source/MovieSceneAnimMixer/Public/MovieSceneAnimTransitionSectionBase.h
+++ b/Engine/Plugins/Experimental/MovieSceneAnimMixer/Source/MovieSceneAnimMixer/Public/MovieSceneAnimTransitionSectionBase.h
@@ -92,6 +92,10 @@ public:
 	//~ IMovieSceneAnimationMixerItemInterface
 	virtual FColor GetMixerItemTint() const override { return TransitionTintColor; }
 
+	virtual bool SupportsVerticalDragging() const override { return false; }
+
+	virtual bool IsVisibleInAddSectionMenu() const override {return false; }
+
 	//~ UMovieSceneSection interface
 	virtual EMovieSceneChannelProxyType CacheChannelProxy() override final;
 
diff --git a/Engine/Plugins/Experimental/MovieSceneAnimMixer/Source/MovieSceneAnimMixerEditor/MovieSceneAnimMixerEditor.Build.cs b/Engine/Plugins/Experimental/MovieSceneAnimMixer/Source/MovieSceneAnimMixerEditor/MovieSceneAnimMixerEditor.Build.cs
--- a/Engine/Plugins/Experimental/MovieSceneAnimMixer/Source/MovieSceneAnimMixerEditor/MovieSceneAnimMixerEditor.Build.cs
+++ b/Engine/Plugins/Experimental/MovieSceneAnimMixer/Source/MovieSceneAnimMixerEditor/MovieSceneAnimMixerEditor.Build.cs
@@ -14,6 +14,9 @@ namespace UnrealBuildTool.Rules
 					"AnimGraphRuntime", 
 					"Core",
 					"CoreUObject",
+					"SequencerCore",
+					"Sequencer",
+					"MovieSceneAnimMixer",
 				}
 			);
 			
diff --git a/Engine/Plugins/Experimental/MovieSceneAnimMixer/Source/MovieSceneAnimMixerEditor/MovieSceneAnimationMaskDecorationEditor.cpp b/Engine/Plugins/Experimental/MovieSceneAnimMixer/Source/MovieSceneAnimMixerEditor/MovieSceneAnimationMaskDecorationEditor.cpp
new file mode 100644
--- /dev/null
+++ b/Engine/Plugins/Experimental/MovieSceneAnimMixer/Source/MovieSceneAnimMixerEditor/MovieSceneAnimationMaskDecorationEditor.cpp
@@ -0,0 +1,18 @@
+﻿// Copyright Epic Games, Inc. All Rights Reserved.
+
+#include "MovieSceneAnimationMaskDecorationEditor.h"
+
+#include "MovieSceneAnimationMaskDecoration.h"
+
+namespace UE::Sequencer
+{
+	UClass* FMovieSceneAnimationMaskDecorationEditor::GetDecorationClass() const
+	{
+		return UMovieSceneAnimationMaskDecoration::StaticClass();
+	}
+
+	TSharedPtr<ISequencerSection> FMovieSceneAnimationMaskDecorationEditor::MakeSectionInterface(UMovieSceneSection& Section, UMovieSceneTrack& OwningTrack, FGuid ObjectBinding)
+	{
+		return MakeShared<FAnimMixerMaskSection>(Section);
+	}
+}
diff --git a/Engine/Plugins/Experimental/MovieSceneAnimMixer/Source/MovieSceneAnimMixerEditor/MovieSceneAnimationMaskDecorationEditor.h b/Engine/Plugins/Experimental/MovieSceneAnimMixer/Source/MovieSceneAnimMixerEditor/MovieSceneAnimationMaskDecorationEditor.h
new file mode 100644
--- /dev/null
+++ b/Engine/Plugins/Experimental/MovieSceneAnimMixer/Source/MovieSceneAnimMixerEditor/MovieSceneAnimationMaskDecorationEditor.h
@@ -0,0 +1,46 @@
+﻿// Copyright Epic Games, Inc. All Rights Reserved.
+
+#pragma once
+#include "ISequencerDecorationEditor.h"
+#include "ISequencerSection.h"
+#include "MovieSceneAnimMixerMaskSection.h"
+#include "MovieSceneTrack.h"
+
+namespace UE::Sequencer
+{
+
+class FAnimMixerMaskSection : public FSequencerSection
+{
+public:
+
+	FAnimMixerMaskSection(UMovieSceneSection& InSection)
+		: FSequencerSection(InSection)
+	{}
+	/** ISequencerSection Interface */
+
+	virtual FText GetSectionTitle() const override
+	{
+		if (TStrongObjectPtr<UMovieSceneSection> Section = WeakSection.Pin())
+		{
+			if (UMovieSceneAnimMixerMaskSection* MaskSection = Cast<UMovieSceneAnimMixerMaskSection>(Section.Get()))
+			{
+				if (const UUAFBlendMask* BlendMask = MaskSection->GetBlendMask().ResolveObjectPtr())
+				{
+					return FText::FromString(BlendMask->GetName());
+				}
+			}
+		}
+		return  FText();
+	}
+};
+
+class FMovieSceneAnimationMaskDecorationEditor : public ISequencerDecorationEditor
+{
+public:
+	virtual UClass* GetDecorationClass() const override;
+
+	virtual TSharedPtr<ISequencerSection> MakeSectionInterface(UMovieSceneSection& Section, UMovieSceneTrack& OwningTrack, FGuid ObjectBinding) override;
+
+};
+
+}
diff --git a/Engine/Plugins/Experimental/MovieSceneAnimMixer/Source/MovieSceneAnimMixerEditor/Private/MVVM/AnimationMixerTrackModel.cpp b/Engine/Plugins/Experimental/MovieSceneAnimMixer/Source/MovieSceneAnimMixerEditor/Private/MVVM/AnimationMixerTrackModel.cpp
--- a/Engine/Plugins/Experimental/MovieSceneAnimMixer/Source/MovieSceneAnimMixerEditor/Private/MVVM/AnimationMixerTrackModel.cpp
+++ b/Engine/Plugins/Experimental/MovieSceneAnimMixer/Source/MovieSceneAnimMixerEditor/Private/MVVM/AnimationMixerTrackModel.cpp
@@ -5,7 +5,7 @@
 #include "MVVM/AnimationMixerLayoutHelpers.h"
 #include "MovieSceneAnimationMixerTrack.h"
 #include "MovieSceneAnimationMixerLayer.h"
-#include "MovieSceneAnimTransitionSectionBase.h"
+#include "Sections/MovieSceneSkeletalAnimationSection.h"
 #include "MVVM/DecorationModelStorageExtension.h"
 #include "MVVM/ViewModels/TrackModel.h"
 #include "MVVM/ViewModels/TrackRowModel.h"
@@ -341,10 +341,15 @@ TOptional<FAnimMixerDropPreview> FAnimationMixerTrackModel::ComputeDropPreview(c
 		return TOptional<FAnimMixerDropPreview>();
 	}
 
-	// Transition sections cannot be moved vertically - they must stay on the same row as from/to sections
+	// Some sections cannot be moved vertically - they must stay on the same row as from/to sections
 	for (UMovieSceneSection* Section : *Context.DraggedSections)
 	{
-		if (Section && Section->IsA<UMovieSceneAnimTransitionSectionBase>())
+		// Skeletal animation sections do not implement the mixer item interface, but are still draggable
+		const UMovieSceneSkeletalAnimationSection* SkeletalAnimationSection = Cast<UMovieSceneSkeletalAnimationSection>(Section);
+
+		// If the section is not a skeletal animation section, check to see if it is a mixer item that supports dragging.
+		const IMovieSceneAnimationMixerItemInterface* MixerItem = Cast<IMovieSceneAnimationMixerItemInterface>(Section);
+		if (!SkeletalAnimationSection && (!MixerItem || !MixerItem->SupportsVerticalDragging()))
 		{
 			return TOptional<FAnimMixerDropPreview>();
 		}
diff --git a/Engine/Plugins/Experimental/MovieSceneAnimMixer/Source/MovieSceneAnimMixerEditor/Private/MovieSceneAnimMixerEditorModule.cpp b/Engine/Plugins/Experimental/MovieSceneAnimMixer/Source/MovieSceneAnimMixerEditor/Private/MovieSceneAnimMixerEditorModule.cpp
--- a/Engine/Plugins/Experimental/MovieSceneAnimMixer/Source/MovieSceneAnimMixerEditor/Private/MovieSceneAnimMixerEditorModule.cpp
+++ b/Engine/Plugins/Experimental/MovieSceneAnimMixer/Source/MovieSceneAnimMixerEditor/Private/MovieSceneAnimMixerEditorModule.cpp
@@ -214,6 +214,9 @@ namespace UE::MovieScene
 
 		RootMotionSettingsDecorationEditor = MakeUnique<UE::Sequencer::FRootMotionSettingsDecorationEditor>();
 		IModularFeatures::Get().RegisterModularFeature(ISequencerDecorationEditor::GetModularFeatureName(), RootMotionSettingsDecorationEditor.Get());
+
+		AnimMaskingDecorationEditor = MakeUnique<FMovieSceneAnimationMaskDecorationEditor>();
+		IModularFeatures::Get().RegisterModularFeature(ISequencerDecorationEditor::GetModularFeatureName(), AnimMaskingDecorationEditor.Get());
 	}
 
 	void FMovieSceneAnimMixerEditorModule::ShutdownModule()
@@ -228,6 +231,10 @@ namespace UE::MovieScene
 		{
 			IModularFeatures::Get().UnregisterModularFeature(ISequencerDecorationEditor::GetModularFeatureName(), RootMotionTargetDecorationEditor.Get());
 		}
+		if (AnimMaskingDecorationEditor)
+		{
+			IModularFeatures::Get().UnregisterModularFeature(ISequencerDecorationEditor::GetModularFeatureName(), AnimMaskingDecorationEditor.Get());
+		}
 
 		// Unregister target menu providers
 		if (AutomaticTargetMenuProvider)
diff --git a/Engine/Plugins/Experimental/MovieSceneAnimMixer/Source/MovieSceneAnimMixerEditor/Private/MovieSceneAnimMixerEditorModule.h b/Engine/Plugins/Experimental/MovieSceneAnimMixer/Source/MovieSceneAnimMixerEditor/Private/MovieSceneAnimMixerEditorModule.h
--- a/Engine/Plugins/Experimental/MovieSceneAnimMixer/Source/MovieSceneAnimMixerEditor/Private/MovieSceneAnimMixerEditorModule.h
+++ b/Engine/Plugins/Experimental/MovieSceneAnimMixer/Source/MovieSceneAnimMixerEditor/Private/MovieSceneAnimMixerEditorModule.h
@@ -5,6 +5,7 @@
 #include "CoreMinimal.h"
 #include "IAnimGraph_SequencerMixerTargetConnector.h"
 #include "Modules/ModuleInterface.h"
+#include "MovieSceneAnimMixerEditor/MovieSceneAnimationMaskDecorationEditor.h"
 
 class FAutomaticTargetMenuProvider;
 class FAnimInstanceTargetMenuProvider;
@@ -42,6 +43,7 @@ namespace UE::MovieScene
 		// Decoration editors
 		TUniquePtr<UE::Sequencer::FRootMotionTargetDecorationEditor> RootMotionTargetDecorationEditor;
 		TUniquePtr<UE::Sequencer::FRootMotionSettingsDecorationEditor> RootMotionSettingsDecorationEditor;
+		TUniquePtr<Sequencer::FMovieSceneAnimationMaskDecorationEditor> AnimMaskingDecorationEditor;
 	};
 }
 
diff --git a/Engine/Plugins/Experimental/MovieSceneAnimMixer/Source/MovieSceneAnimMixerEditor/Private/MovieSceneAnimationMixerTrackEditor.cpp b/Engine/Plugins/Experimental/MovieSceneAnimMixer/Source/MovieSceneAnimMixerEditor/Private/MovieSceneAnimationMixerTrackEditor.cpp
--- a/Engine/Plugins/Experimental/MovieSceneAnimMixer/Source/MovieSceneAnimMixerEditor/Private/MovieSceneAnimationMixerTrackEditor.cpp
+++ b/Engine/Plugins/Experimental/MovieSceneAnimMixer/Source/MovieSceneAnimMixerEditor/Private/MovieSceneAnimationMixerTrackEditor.cpp
@@ -20,6 +20,7 @@
 #include "MovieSceneAnimationMixerItemInterface.h"
 #include "MovieSceneAnimationMixerLayer.h"
 #include "MovieSceneAnimationMixerTrack.h"
+#include "MovieSceneAnimMixerMaskSection.h"
 #include "MovieSceneTrack.h"
 #include "MVVM/Extensions/IObjectBindingExtension.h"
 #include "MVVM/Extensions/ITrackExtension.h"
@@ -943,10 +944,14 @@ TSharedRef<SWidget> FAnimationMixerTrackEditor::BuildAddSectionSubMenu(TWeakView
 		{
 			It.RemoveCurrent();
 		}
-		else if (Class->IsChildOf(UMovieSceneAnimTransitionSectionBase::StaticClass()))
+		else
 		{
-			// Transition sections are created via context menu, not the add menu
-			It.RemoveCurrent();
+			IMovieSceneAnimationMixerItemInterface* MixerItem = Cast<IMovieSceneAnimationMixerItemInterface>(Class->GetDefaultObject());
+			if (MixerItem && !MixerItem->IsVisibleInAddSectionMenu())
+			{
+				// Transition sections are created via context menu, not the add menu, and decoration owned sections are added in the decoration menu
+				It.RemoveCurrent();
+			}
 		}
 	}
 
@@ -1091,20 +1096,20 @@ TSharedRef<SWidget> FAnimationMixerTrackEditor::BuildAddSectionSubMenu(TWeakView
 		UMovieSceneAnimationMixerLayer* Layer = MixerTrack->GetLayer(LayerExtension->GetRowIndex());
 		if (Layer)
 		{
-			SequencerHelpers::BuildDecorationMenu(MenuBuilder, Layer);
+			SequencerHelpers::BuildDecorationMenu(MenuBuilder, Layer, Sequencer);
 
 			// For single-section layers, also show decorations compatible with the section
 			TArray<UMovieSceneSection*> LayerSections = Layer->GetSections();
 			if (LayerSections.Num() == 1 && LayerSections[0])
 			{
-				SequencerHelpers::BuildDecorationMenu(MenuBuilder, LayerSections[0]);
+				SequencerHelpers::BuildDecorationMenu(MenuBuilder, LayerSections[0], Sequencer);
 			}
 		}
 	}
 	else
 	{
 		// Track-level: add decorations to the track
-		SequencerHelpers::BuildDecorationMenu(MenuBuilder, MixerTrack);
+		SequencerHelpers::BuildDecorationMenu(MenuBuilder, MixerTrack, Sequencer);
 	}
 
 	return MenuBuilder.MakeWidget();
diff --git a/Engine/Source/Editor/Sequencer/Private/MVVM/ViewModels/DecorationOutlinerModel.cpp b/Engine/Source/Editor/Sequencer/Private/MVVM/ViewModels/DecorationOutlinerModel.cpp
--- a/Engine/Source/Editor/Sequencer/Private/MVVM/ViewModels/DecorationOutlinerModel.cpp
+++ b/Engine/Source/Editor/Sequencer/Private/MVVM/ViewModels/DecorationOutlinerModel.cpp
@@ -20,6 +20,8 @@
 #include "MVVM/ViewModels/TrackModelLayoutBuilder.h"
 #include "MVVM/ViewModels/ViewDensity.h"
 #include "SSequencer.h"
+#include "MVVM/ViewModels/OutlinerColumns/OutlinerColumnTypes.h"
+#include "MVVM/Views/ViewUtilities.h"
 
 #define LOCTEXT_NAMESPACE "DecorationOutlinerModel"
 
@@ -38,6 +40,15 @@ FDecorationOutlinerModel::FDecorationOutlinerModel(UMovieSceneTrack* InParentTra
 
 FDecorationOutlinerModel::~FDecorationOutlinerModel()
 {
+	if (UMovieSceneSignedObject* SignedObject = Cast<UMovieSceneSignedObject>(WeakDecoration.Get()))
+	{
+		SignedObject->OnSignatureChanged().Remove(SignatureChangedHandle);
+	}
+}
+
+void FDecorationOutlinerModel::OnDecorationSignatureChanged()
+{
+	bDecorationSignatureChanged = true;
 }
 
 bool FDecorationOutlinerModel::NeedsReconstruct() const
@@ -56,6 +67,11 @@ bool FDecorationOutlinerModel::NeedsReconstruct() const
 		}
 	}
 
+	if (IMovieSceneSectionProviderDecoration* SectionProviderDecoration = Cast<IMovieSceneSectionProviderDecoration>(Decoration))
+	{
+		return SectionProviderDecoration->ConsumeStructureChanged();
+	}
+
 	return false;
 }
 
@@ -297,9 +313,40 @@ FOutlinerSizing FDecorationOutlinerModel::GetOutlinerSizing() const
 	return FOutlinerSizing(Height);
 }
 
-TSharedPtr<SWidget> FDecorationOutlinerModel::CreateOutlinerViewForColumn(const FCreateOutlinerViewParams& InParams,
-																						const FName& InColumnName)
+TSharedRef<SWidget> FDecorationOutlinerModel::BuildAddDecorationSubMenu(TSharedPtr<FEditorViewModel> Editor) const
+{
+	FMenuBuilder MenuBuilder(true, nullptr);
+	const TObjectPtr<UObject> Decoration = WeakDecoration.Get();
+
+	const TSharedPtr<FSequencerEditorViewModel> SequencerEditor = Editor->CastThisShared<FSequencerEditorViewModel>();
+	if (!SequencerEditor)
+	{
+		return MenuBuilder.MakeWidget();
+	}
+
+	if (Decoration)
+	{
+		UMovieSceneDecorationContainerObject* DecorationContainer = Decoration->GetTypedOuter<UMovieSceneDecorationContainerObject>();
+		SequencerHelpers::BuildDecorationMenu(MenuBuilder, DecorationContainer, SequencerEditor->GetSequencer());
+	}
+
+	return  MenuBuilder.MakeWidget();
+}
+
+TSharedPtr<SWidget> FDecorationOutlinerModel::CreateOutlinerViewForColumn(const FCreateOutlinerViewParams& InParams, const FName& InColumnName)
 {
+	if (InColumnName == FCommonOutlinerNames::Add)
+	{
+		return MakeAddButton(
+			LOCTEXT("AddModifier", "Modifier"),
+			FOnGetContent::CreateSP(
+				this,
+				&FDecorationOutlinerModel::BuildAddDecorationSubMenu,
+				InParams.Editor
+				),
+				InParams.Editor
+			);
+	}
 	if (TrackEditor)
 	{
 		const FBuildColumnWidgetParams Params(SharedThis(this), InParams);
diff --git a/Engine/Source/Editor/Sequencer/Private/MVVM/ViewModels/SectionOutlinerModel.cpp b/Engine/Source/Editor/Sequencer/Private/MVVM/ViewModels/SectionOutlinerModel.cpp
--- a/Engine/Source/Editor/Sequencer/Private/MVVM/ViewModels/SectionOutlinerModel.cpp
+++ b/Engine/Source/Editor/Sequencer/Private/MVVM/ViewModels/SectionOutlinerModel.cpp
@@ -299,12 +299,12 @@ TSharedPtr<SWidget> FSectionOutlinerModel::CreateOutlinerViewForColumn(const FCr
 		if (Section && SequencerHelpers::HasCompatibleDecorations(Section))
 		{
 			TWeakObjectPtr<UMovieSceneSection> WeakSectionForMenu = Section;
-			auto GetMenuContent = [WeakSectionForMenu]() -> TSharedRef<SWidget>
+			auto GetMenuContent = [WeakSectionForMenu, Editor]() -> TSharedRef<SWidget>
 			{
 				FMenuBuilder MenuBuilder(true, nullptr);
 				if (UMovieSceneSection* PinnedSection = WeakSectionForMenu.Get())
 				{
-					SequencerHelpers::BuildDecorationMenu(MenuBuilder, PinnedSection);
+					SequencerHelpers::BuildDecorationMenu(MenuBuilder, PinnedSection, Editor->GetSequencer());
 				}
 				return MenuBuilder.MakeWidget();
 			};
diff --git a/Engine/Source/Editor/Sequencer/Private/MovieSceneTrackEditor.cpp b/Engine/Source/Editor/Sequencer/Private/MovieSceneTrackEditor.cpp
--- a/Engine/Source/Editor/Sequencer/Private/MovieSceneTrackEditor.cpp
+++ b/Engine/Source/Editor/Sequencer/Private/MovieSceneTrackEditor.cpp
@@ -364,7 +364,7 @@ TSharedPtr<SWidget> FMovieSceneTrackEditor::BuildOutlinerEditWidget(const FGuid&
 			// Add decoration entries if track supports decorations
 			if (bHasDecorations)
 			{
-				SequencerHelpers::BuildDecorationMenu(MenuBuilder, Track);
+				SequencerHelpers::BuildDecorationMenu(MenuBuilder, Track, WeakSequencer);
 			}
 
 			return MenuBuilder.MakeWidget();
diff --git a/Engine/Source/Editor/Sequencer/Private/SequencerCommonHelpers.cpp b/Engine/Source/Editor/Sequencer/Private/SequencerCommonHelpers.cpp
--- a/Engine/Source/Editor/Sequencer/Private/SequencerCommonHelpers.cpp
+++ b/Engine/Source/Editor/Sequencer/Private/SequencerCommonHelpers.cpp
@@ -1191,7 +1191,7 @@ bool SequencerHelpers::HasOutlinerVisibleDecorations(const UMovieSceneDecoration
 
 void SequencerHelpers::BuildDecorationMenu(
 	FMenuBuilder& MenuBuilder,
-	UMovieSceneDecorationContainerObject* DecorationContainer)
+	UMovieSceneDecorationContainerObject* DecorationContainer, TWeakPtr<ISequencer> InSequencer)
 {
 	if (!DecorationContainer)
 	{
@@ -1210,13 +1210,16 @@ void SequencerHelpers::BuildDecorationMenu(
 	{
 		for (UClass* DecorationClass : CompatibleDecorations)
 		{
-			bool bAlreadyExists = DecorationContainer->GetDecorations().ContainsByPredicate(
+			const TObjectPtr<UObject>* ExistingDecoration = DecorationContainer->GetDecorations().FindByPredicate(
 				[DecorationClass](const TObjectPtr<UObject>& Existing)
 				{
 					return Existing && Existing->GetClass() == DecorationClass;
 				});
 
-			if (!bAlreadyExists)
+			IMovieSceneSectionProviderDecoration* SectionProvider = ExistingDecoration ? Cast<IMovieSceneSectionProviderDecoration>(ExistingDecoration->Get()) : nullptr;
+			const bool bSupportsMultipleSections = SectionProvider ? SectionProvider->SupportsMultipleSections() : false;
+
+			if (!ExistingDecoration)
 			{
 				MenuBuilder.AddMenuEntry(
 					DecorationClass->GetDisplayNameText(),
@@ -1235,6 +1238,50 @@ void SequencerHelpers::BuildDecorationMenu(
 					}))
 				);
 			}
+			else if (bSupportsMultipleSections)
+			{
+				MenuBuilder.AddMenuEntry(
+					DecorationClass->GetDisplayNameText(),
+					DecorationClass->GetToolTipText(),
+					FSlateIcon(),
+					FUIAction(FExecuteAction::CreateLambda([DecorationContainer, DecorationClass, SectionProvider, InSequencer]()
+					{
+						FScopedTransaction Transaction(
+							FText::Format(NSLOCTEXT("SequencerHelpers", "AddDecorationSection", "Add {0}"),
+							DecorationClass->GetDisplayNameText()));
+
+						TSharedPtr<ISequencer> Sequencer = InSequencer.IsValid() ? InSequencer.Pin() : nullptr;
+						if (!Sequencer)
+						{
+							return;
+						}
+
+						FQualifiedFrameTime CurrentTime = Sequencer->GetLocalTime();
+						FFrameNumber PlaybackEnd = UE::MovieScene::DiscreteExclusiveUpper(Sequencer->GetFocusedMovieSceneSequence()->GetMovieScene()->GetPlaybackRange());
+
+						DecorationContainer->Modify();
+						UMovieSceneSection* NewSection = SectionProvider->CreateNewSection();
+						if (Sequencer->GetInfiniteKeyAreas() && NewSection->GetSupportsInfiniteRange())
+						{
+							NewSection->SetRange(TRange<FFrameNumber>::All());
+						}
+						else
+						{
+							FFrameNumber NewSectionRangeEnd = PlaybackEnd;
+							if (PlaybackEnd <= CurrentTime.Time.FrameNumber)
+							{
+								const FAnimatedRange ViewRange = Sequencer->GetViewRange();
+								const FFrameRate TickResolution = Sequencer->GetFocusedTickResolution();
+								NewSectionRangeEnd = (ViewRange.GetUpperBoundValue() * TickResolution).FloorToFrame();
+							}
+
+							NewSection->SetRange(TRange<FFrameNumber>(CurrentTime.Time.FrameNumber, NewSectionRangeEnd));
+						}
+						SectionProvider->AddSection(NewSection);
+						DecorationContainer->MarkAsChanged();
+					}))
+				);
+			}
 		}
 	}
 	MenuBuilder.EndSection();
diff --git a/Engine/Source/Editor/Sequencer/Public/MVVM/ViewModels/DecorationOutlinerModel.h b/Engine/Source/Editor/Sequencer/Public/MVVM/ViewModels/DecorationOutlinerModel.h
--- a/Engine/Source/Editor/Sequencer/Public/MVVM/ViewModels/DecorationOutlinerModel.h
+++ b/Engine/Source/Editor/Sequencer/Public/MVVM/ViewModels/DecorationOutlinerModel.h
@@ -37,13 +37,14 @@ public:
 		, ITrackAreaExtension
 		, ILockableExtension
 		, IDeletableExtension
-		, ITopLevelChannelHolderExtension);
+		, ITopLevelChannelHolderExtension)
 
 	UE_API explicit FDecorationOutlinerModel(UMovieSceneTrack* InParentTrack, UMovieSceneDecorationContainerObject* InDecorationContainer, UObject* InDecoration);
 	UE_API virtual ~FDecorationOutlinerModel() override;
 
 	/* FViewModel Interface */
-	UE_API virtual void OnConstruct() override;
+	UE_API void OnDecorationSignatureChanged();
+	virtual void OnConstruct() override;
 
 	/* FOutlinerItemModel */
 	UE_API virtual void BuildContextMenu(FMenuBuilder& MenuBuilder) override;
@@ -51,7 +52,8 @@ public:
 
 	/* IOutlinerExtension */
 	UE_API virtual FOutlinerSizing GetOutlinerSizing() const override;
-	UE_API virtual TSharedPtr<SWidget> CreateOutlinerViewForColumn(const FCreateOutlinerViewParams& InParams, const FName& InColumnName) override;
+	UE_API TSharedRef<SWidget> BuildAddDecorationSubMenu(TSharedPtr<FEditorViewModel> EditorViewModel) const;
+	virtual TSharedPtr<SWidget> CreateOutlinerViewForColumn(const FCreateOutlinerViewParams& InParams, const FName& InColumnName) override;
 	UE_API virtual FText GetLabel() const override;
 
 	/* ISectionOwnerExtension */
@@ -73,7 +75,7 @@ public:
 	UE_API virtual ELockableLockState GetLockState() const override;
 	UE_API virtual void SetIsLocked(bool bIsLocked) override;
 
-	UE_API TWeakObjectPtr<UMovieSceneDecorationContainerObject> GetDecorationContainer() const { return WeakDecorationContainer; }
+	TWeakObjectPtr<UMovieSceneDecorationContainerObject> GetDecorationContainer() const { return WeakDecorationContainer; }
 
 	// Get the decoration object this model represents. Used for sidebar property display.
 	UE_API UObject* GetDecoration() const;
@@ -104,6 +106,12 @@ private:
 	// Tracks channel models added to section models so they can be removed during cleanup.
 	TArray<TWeakViewModelPtr<FChannelModel>> TrackedChannelModels;
 
+	// Handle for this decoration's on signature changed delegate.
+	FDelegateHandle SignatureChangedHandle;
+
+	bool bDecorationSignatureChanged = false;
+
+	FGuid DecorationSignature;
 };
 
 #undef UE_API
diff --git a/Engine/Source/Editor/Sequencer/Public/SequencerCommonHelpers.h b/Engine/Source/Editor/Sequencer/Public/SequencerCommonHelpers.h
--- a/Engine/Source/Editor/Sequencer/Public/SequencerCommonHelpers.h
+++ b/Engine/Source/Editor/Sequencer/Public/SequencerCommonHelpers.h
@@ -156,9 +156,9 @@ public:
 	/**
 	 * Populates a menu with "Add Decoration" entries for compatible decorations.
 	 */
-	static void BuildDecorationMenu(
-		FMenuBuilder& MenuBuilder,
-		UMovieSceneDecorationContainerObject* DecorationContainer);
+	static void BuildDecorationMenu(FMenuBuilder& MenuBuilder
+		, UMovieSceneDecorationContainerObject* DecorationContainer
+		, TWeakPtr<ISequencer> InSequencer);
 };
 
 #undef LOCTEXT_NAMESPACE
diff --git a/Engine/Source/Runtime/MovieSceneTracks/Public/MovieSceneAnimationMixerItemInterface.h b/Engine/Source/Runtime/MovieSceneTracks/Public/MovieSceneAnimationMixerItemInterface.h
--- a/Engine/Source/Runtime/MovieSceneTracks/Public/MovieSceneAnimationMixerItemInterface.h
+++ b/Engine/Source/Runtime/MovieSceneTracks/Public/MovieSceneAnimationMixerItemInterface.h
@@ -34,4 +34,10 @@ public:
 
 	/** Icon for add menu entry (default: no icon) */
 	virtual FSlateIcon GetIcon() const { return FSlateIcon(); }
+
+	/** If this item supports vertical dragging within the mixer */
+	virtual bool SupportsVerticalDragging() const { return true; }
+
+	/* Some items such as decorations and transitions are not added through the section add menu */
+	virtual bool IsVisibleInAddSectionMenu() const { return true; }
 };
