Sequencer: Unify property handle to property path conversion

Result should be:
Outer array and inner separately, ie. MyArray.MyArray[0].bVisible

This fixes a bug when generating a track/section/key from the diamond in the details panel because it was incorrectly converting the property handle without the outer array, ie. MyArray[0].bVisible

#Synced-CL 51597275

#Synced-CL 51597275

#Synced-CL 51630155

diff --git a/Engine/Source/Editor/Sequencer/Private/KeyPropertyParams.cpp b/Engine/Source/Editor/Sequencer/Private/KeyPropertyParams.cpp
--- a/Engine/Source/Editor/Sequencer/Private/KeyPropertyParams.cpp
+++ b/Engine/Source/Editor/Sequencer/Private/KeyPropertyParams.cpp
@@ -9,41 +9,7 @@
 #include "Templates/SharedPointer.h"
 #include "UObject/Class.h"
 #include "UObject/Object.h"
-
-FPropertyPath PropertyHandleToPropertyPath(const UClass* OwnerClass, const IPropertyHandle& InPropertyHandle)
-{
-	TArray<FPropertyInfo> PropertyInfo;
-	PropertyInfo.Insert(FPropertyInfo(InPropertyHandle.GetProperty(), InPropertyHandle.GetIndexInArray()), 0);
-
-	TSharedPtr<IPropertyHandle> CurrentHandle = InPropertyHandle.GetParentHandle();
-	while (CurrentHandle.IsValid() && CurrentHandle->GetProperty() != nullptr)
-	{
-		// IPropertyHandle chains contain arrays in a manner designed for display in the property editor, e.g.
-		// Container.Array.Array[ArrayIndex].StructInner.
-		// We need to collapse adjacent array properties as we are looking for Container.Array[ArrayIndex].StructInner
-		// to for a well-formed 'property path'
-		TSharedPtr<IPropertyHandle> ParentHandle = CurrentHandle->GetParentHandle();
-		if (ParentHandle.IsValid() && ParentHandle->GetProperty() != nullptr && ParentHandle->GetProperty()->IsA<FArrayProperty>())
-		{
-			FArrayProperty* ArrayProperty = CastFieldChecked<FArrayProperty>(ParentHandle->GetProperty());
-
-			PropertyInfo.Insert(FPropertyInfo(ParentHandle->GetProperty(), CurrentHandle->GetIndexInArray()), 0);
-			CurrentHandle = ParentHandle->GetParentHandle();
-		}
-		else
-		{
-			PropertyInfo.Insert(FPropertyInfo(CurrentHandle->GetProperty(), CurrentHandle->GetIndexInArray()), 0);
-			CurrentHandle = CurrentHandle->GetParentHandle();
-		}
-	}
-
-	FPropertyPath PropertyPath;
-	for (const FPropertyInfo& Info : PropertyInfo)
-	{
-		PropertyPath.AddProperty(Info);
-	}
-	return PropertyPath;
-}
+#include "SequencerCommonHelpers.h"
 
 FCanKeyPropertyParams::FCanKeyPropertyParams(const UClass* InObjectClass, const FPropertyPath& InPropertyPath)
 	: ObjectClass(InObjectClass)
@@ -53,7 +19,7 @@ FCanKeyPropertyParams::FCanKeyPropertyParams(const UClass* InObjectClass, const
 
 FCanKeyPropertyParams::FCanKeyPropertyParams(const UClass* InObjectClass, const IPropertyHandle& InPropertyHandle)
 	: ObjectClass(InObjectClass)
-	, PropertyPath(PropertyHandleToPropertyPath(InObjectClass, InPropertyHandle))
+	, PropertyPath(SequencerHelpers::PropertyHandleToPropertyPath(InPropertyHandle))
 {
 }
 
@@ -115,7 +81,7 @@ FKeyPropertyParams::FKeyPropertyParams(TArray<UObject*> InObjectsToKey, const FP
 
 FKeyPropertyParams::FKeyPropertyParams(TArray<UObject*> InObjectsToKey, const IPropertyHandle& InPropertyHandle, ESequencerKeyMode InKeyMode)
 	: ObjectsToKey(InObjectsToKey)
-	, PropertyPath(InObjectsToKey.Num() > 0 ? PropertyHandleToPropertyPath(InObjectsToKey[0]->GetClass(), InPropertyHandle) : FPropertyPath())
+	, PropertyPath(InObjectsToKey.Num() > 0 ? SequencerHelpers::PropertyHandleToPropertyPath(InPropertyHandle) : FPropertyPath())
 	, KeyMode(InKeyMode)
 {
 }
diff --git a/Engine/Source/Editor/Sequencer/Private/SequencerCommonHelpers.cpp b/Engine/Source/Editor/Sequencer/Private/SequencerCommonHelpers.cpp
--- a/Engine/Source/Editor/Sequencer/Private/SequencerCommonHelpers.cpp
+++ b/Engine/Source/Editor/Sequencer/Private/SequencerCommonHelpers.cpp
@@ -36,6 +36,7 @@
 #include "MovieSceneSequence.h"
 #include "Sections/MovieSceneSubSection.h"
 #include "PropertyEditorModule.h"
+#include "PropertyHandle.h"
 #include "PropertyPermissionList.h"
 #include "SSequencer.h"
 #include "Sequencer.h"
@@ -1286,3 +1287,30 @@ void SequencerHelpers::BuildDecorationMenu(
 	}
 	MenuBuilder.EndSection();
 }
+
+FPropertyPath SequencerHelpers::PropertyHandleToPropertyPath(const IPropertyHandle& InPropertyHandle)
+{
+	TSharedRef<const IPropertyHandle> CurrentHandle = InPropertyHandle.AsShared();
+	TSharedPtr<const IPropertyHandle> ParentHandle = CurrentHandle->GetParentHandle();
+
+	TArray<FPropertyInfo> PropertyInfos;
+	if (CurrentHandle->GetProperty())
+	{
+		PropertyInfos.Emplace(CurrentHandle->GetProperty(), CurrentHandle->GetArrayIndex());
+	}
+
+	while (ParentHandle.IsValid() && ParentHandle->GetProperty())
+	{
+		PropertyInfos.Emplace(ParentHandle->GetProperty(), ParentHandle->GetArrayIndex());
+		CurrentHandle = ParentHandle.ToSharedRef();
+		ParentHandle = ParentHandle->GetParentHandle();
+	}
+
+	FPropertyPath PropertyPath;
+	for (const FPropertyInfo& PropertyInfo : ReverseIterate(PropertyInfos))
+	{
+		PropertyPath.AddProperty(PropertyInfo);
+	}
+
+	return PropertyPath;
+}
\ No newline at end of file
diff --git a/Engine/Source/Editor/Sequencer/Private/SequencerPropertyKeyedStatus.cpp b/Engine/Source/Editor/Sequencer/Private/SequencerPropertyKeyedStatus.cpp
--- a/Engine/Source/Editor/Sequencer/Private/SequencerPropertyKeyedStatus.cpp
+++ b/Engine/Source/Editor/Sequencer/Private/SequencerPropertyKeyedStatus.cpp
@@ -6,6 +6,7 @@
 #include "ISequencerObjectChangeListener.h"
 #include "MovieScene.h"
 #include "PropertyHandle.h"
+#include "SequencerCommonHelpers.h"
 #include "Tracks/MovieScenePropertyTrack.h"
 
 FSequencerPropertyKeyedStatusHandler::FPropertyParameters::FPropertyParameters(ISequencer& InSequencer, const IPropertyHandle& ActualProperty)
@@ -14,39 +15,12 @@ FSequencerPropertyKeyedStatusHandler::FPropertyParameters::FPropertyParameters(I
 {
 	CurrentFrameRange = TRange<FFrameNumber>(Sequencer.GetLocalTime().Time.FrameNumber);
 
-	FPropertyPath PropertyPath = BuildPropertyPath();
+	FPropertyPath PropertyPath = SequencerHelpers::PropertyHandleToPropertyPath(ActualProperty);
 	TrackPropertyName = FindTrackPropertyName(PropertyPath);
 	SubPropertyPath = BuildSubPropertyPath();
 	TrackPropertyStructName = FindTrackPropertyStructName(PropertyPath);
 }
 
-FPropertyPath FSequencerPropertyKeyedStatusHandler::FPropertyParameters::BuildPropertyPath() const
-{
-	TSharedRef<const IPropertyHandle> CurrentHandle = ActualProperty.AsShared();
-	TSharedPtr<const IPropertyHandle> ParentHandle = ActualProperty.GetParentHandle();
-
-	TArray<FPropertyInfo> PropertyInfos;
-	if (CurrentHandle->GetProperty())
-	{
-		PropertyInfos.Emplace(CurrentHandle->GetProperty(), CurrentHandle->GetArrayIndex());
-	}
-
-	while (ParentHandle.IsValid() && ParentHandle->GetProperty())
-	{
-		PropertyInfos.Emplace(ParentHandle->GetProperty(), ParentHandle->GetArrayIndex());
-		CurrentHandle = ParentHandle.ToSharedRef();
-		ParentHandle = ParentHandle->GetParentHandle();
-	}
-
-	FPropertyPath PropertyPath;
-	for (const FPropertyInfo& PropertyInfo : ReverseIterate(PropertyInfos))
-	{
-		PropertyPath.AddProperty(PropertyInfo);
-	}
-
-	return PropertyPath;
-}
-
 FString FSequencerPropertyKeyedStatusHandler::FPropertyParameters::FindTrackPropertyName(FPropertyPath& InOutPropertyPath) const
 {
 	const UClass* ObjectClass = ActualProperty.GetOuterBaseClass();
diff --git a/Engine/Source/Editor/Sequencer/Public/SequencerCommonHelpers.h b/Engine/Source/Editor/Sequencer/Public/SequencerCommonHelpers.h
--- a/Engine/Source/Editor/Sequencer/Public/SequencerCommonHelpers.h
+++ b/Engine/Source/Editor/Sequencer/Public/SequencerCommonHelpers.h
@@ -11,8 +11,10 @@
 
 #define LOCTEXT_NAMESPACE "SequencerHelpers"
 
+class FPropertyPath;
 class FSequencer;
 class IKeyArea;
+class IPropertyHandle;
 class ISequencer;
 class UMovieSceneDecorationContainerObject;
 class UMovieSceneSection;
@@ -159,6 +161,11 @@ public:
 	static void BuildDecorationMenu(FMenuBuilder& MenuBuilder
 		, UMovieSceneDecorationContainerObject* DecorationContainer
 		, TWeakPtr<ISequencer> InSequencer);
+
+	/*
+	 * Build an FPropertyPath from an IPropertyHandle. 
+	 */
+	static FPropertyPath PropertyHandleToPropertyPath(const IPropertyHandle& InPropertyHandle);
 };
 
 #undef LOCTEXT_NAMESPACE
