Sequencer: When Curve Editor filtering is unlinked, then the USequencerSettings::bIncludePinnedInFilter will be ignored. Effectively, this means that
- when Curve Editor is unlinked, items you pin in Sequencer's Outliner will still be filtered in Curve Editor when unlinked.
- when Curve Editor is linked, items you pin in Sequencer's Outliner are filtered depending on the USequencerSettings::bIncludePinnedInFilter setting.

#Synced-CL 51626380

diff --git a/Engine/Source/Editor/Sequencer/Private/Filters/FilterBarFlags.h b/Engine/Source/Editor/Sequencer/Private/Filters/FilterBarFlags.h
new file mode 100644
--- /dev/null
+++ b/Engine/Source/Editor/Sequencer/Private/Filters/FilterBarFlags.h
@@ -0,0 +1,20 @@
+﻿// Copyright Epic Games, Inc. All Rights Reserved.
+
+#pragma once
+
+#include "Containers/UnrealString.h"
+#include "HAL/Platform.h"
+#include "Misc/EnumClassFlags.h"
+
+namespace UE::Sequencer
+{
+/** Flags that affect filtering */
+enum class EFilterFlags : uint8
+{
+	None = 0,
+	
+	/** If set, then USequencerSettings::bIncludePinnedInFilter is allowed to effect. If the setting is disabled, pinned items skip filters. */
+	PinnedItemsCanSkipFiltering = 1 << 0
+};
+ENUM_CLASS_FLAGS(EFilterFlags);
+} // namespace UE::Sequencer
diff --git a/Engine/Source/Editor/Sequencer/Private/Filters/FilterBarInitArgs.h b/Engine/Source/Editor/Sequencer/Private/Filters/FilterBarInitArgs.h
new file mode 100644
--- /dev/null
+++ b/Engine/Source/Editor/Sequencer/Private/Filters/FilterBarInitArgs.h
@@ -0,0 +1,35 @@
+﻿// Copyright Epic Games, Inc. All Rights Reserved.
+
+#pragma once
+
+#include "FilterBarFlags.h"
+#include "Containers/UnrealString.h"
+#include "HAL/Platform.h"
+
+class FSequencer;
+
+namespace UE::Sequencer
+{
+/** Init args for FSequencerFilterBar. */
+struct FFilterBarInitArgs
+{
+	/** The Sequencer this filter bar is interacting with */
+	FSequencer& Sequencer;
+	
+	/**
+	 * Used for per filter bar instance config.
+	 * This is relevant for the settings to use for reading and saving config.
+	 * @see FSequencerFilterBarConfig, USequencerSettings::FindTrackFilterBar, FSequencerFilterBar::GetIdentifier.
+	 */
+	FString SubIdentifier;
+	
+	/** Affects how filtering is performed. */
+	EFilterFlags Flags = EFilterFlags::None;
+
+	explicit FFilterBarInitArgs(FSequencer& InSequencer UE_LIFETIMEBOUND, const FString& InSubIdentifier, EFilterFlags InFlags = EFilterFlags::None)
+		: Sequencer(InSequencer)
+		, SubIdentifier(InSubIdentifier)
+		, Flags(InFlags)
+	{}
+};
+} // namespace UE::Sequencer
diff --git a/Engine/Source/Editor/Sequencer/Private/Filters/SequencerFilterBar.cpp b/Engine/Source/Editor/Sequencer/Private/Filters/SequencerFilterBar.cpp
--- a/Engine/Source/Editor/Sequencer/Private/Filters/SequencerFilterBar.cpp
+++ b/Engine/Source/Editor/Sequencer/Private/Filters/SequencerFilterBar.cpp
@@ -2,6 +2,7 @@
 
 #include "Filters/SequencerFilterBar.h"
 
+#include "FilterBarInitArgs.h"
 #include "Filters/Filters/SequencerTrackFilter_Keyed.h"
 #include "Filters/Filters/SequencerTrackFilter_Condition.h"
 #include "Filters/Filters/SequencerTrackFilter_CustomText.h"
@@ -39,8 +40,8 @@ using namespace UE::Sequencer;
 
 int32 FSequencerFilterBar::InstanceCount = 0;
 
-FSequencerFilterBar::FSequencerFilterBar(FSequencer& InSequencer, FString InSubIdentifier)
-	: Sequencer(InSequencer)
+FSequencerFilterBar::FSequencerFilterBar(FFilterBarInitArgs InArgs)
+	: Sequencer(InArgs.Sequencer)
 	, CommandList(MakeShared<FUICommandList>())
 	, ClassTypeCategory(MakeShared<FFilterCategory>(LOCTEXT("ActorTypeFilterCategory", "Actor Type Filters"), FText()))
 	, ComponentTypeCategory(MakeShared<FFilterCategory>(LOCTEXT("ObjectTypeFilterCategory", "Object Type Filters"), FText()))
@@ -55,7 +56,8 @@ FSequencerFilterBar::FSequencerFilterBar(FSequencer& InSequencer, FString InSubI
 	, SelectedFilter(MakeShared<FSequencerTrackFilter_Selected>(*this, MiscCategory))
 	, ModifiedFilter(MakeShared<FSequencerTrackFilter_Modified>(*this, MiscCategory))
 	, FilterData(FString())
-	, SubIdentifier(MoveTemp(InSubIdentifier))
+	, SubIdentifier(MoveTemp(InArgs.SubIdentifier))
+	, Flags(InArgs.Flags)
 {
 	InstanceCount++;
 
@@ -1477,7 +1479,7 @@ bool FSequencerFilterBar::FilterNodesRecursive(const bool bInHasActiveFilter, co
 	}
 
 	// Pinning overrides all other filters
-	if (!SequencerSettings->GetIncludePinnedInFilter())
+	if (EnumHasAnyFlags(Flags, EFilterFlags::PinnedItemsCanSkipFiltering) && !SequencerSettings->GetIncludePinnedInFilter())
 	{
 		const TSharedPtr<IPinnableExtension> Pinnable = InStartNode.AsModel()->FindAncestorOfType<IPinnableExtension>(true);
 		if (Pinnable.IsValid() && Pinnable->IsPinned())
diff --git a/Engine/Source/Editor/Sequencer/Private/Filters/SequencerFilterBar.h b/Engine/Source/Editor/Sequencer/Private/Filters/SequencerFilterBar.h
--- a/Engine/Source/Editor/Sequencer/Private/Filters/SequencerFilterBar.h
+++ b/Engine/Source/Editor/Sequencer/Private/Filters/SequencerFilterBar.h
@@ -2,6 +2,7 @@
 
 #pragma once
 
+#include "FilterBarInitArgs.h"
 #include "Filters/CustomTextFilters.h"
 #include "Filters/FilterBarDelegates.h"
 #include "Filters/ISequencerTrackFilters.h"
@@ -10,6 +11,8 @@
 
 class FSequencer;
 class FSequencerTrackFilter;
+class FSequencerTrackFilterCollection;
+class FSequencerTrackFilterMenu;
 class FSequencerTrackFilter_CustomText;
 class FSequencerTrackFilter_Group;
 class FSequencerTrackFilter_HideIsolate;
@@ -17,19 +20,14 @@ class FSequencerTrackFilter_Level;
 class FSequencerTrackFilter_Modified;
 class FSequencerTrackFilter_Selected;
 class FSequencerTrackFilter_Text;
-class FSequencerTrackFilterCollection;
-class FSequencerTrackFilterMenu;
 class FUICommandList;
 class SComboButton;
 class SFilterBarIsolateHideShow;
 class SSequencerFilterBar;
 class SSequencerSearchBox;
 enum class EFilterBarLayout : uint8;
-
-namespace UE::Sequencer
-{
-class IOutlinerExtension;
-}
+namespace UE::Sequencer { class IOutlinerExtension; }
+namespace UE::Sequencer { struct FFilterBarInitArgs; }
 
 enum class ESequencerFilterChange : uint8
 {
@@ -47,13 +45,7 @@ public:
 	DECLARE_EVENT_TwoParams(FSequencerFilterBar, FSequencerFiltersChanged, const ESequencerFilterChange /*InChangeType*/, const TSharedRef<FSequencerTrackFilter>& /*InFilter*/);
 	DECLARE_EVENT_TwoParams(FSequencerFilterBar, FSequencerCustomTextFiltersChanged, const ESequencerFilterChange /*InChangeType*/, const TSharedRef<FSequencerTrackFilter_CustomText>& /*InFilter*/);
 	
-	/** 
-	 * @param InSequencer The owning Sequencer instance.
-	 * @param InSubIdentifier Sub-identifier that affects the result of FSequencerFilterBar::GetIdentifier. 
-	 *		This is relevant for saving per-filter bar instance config.
-	 * @see FSequencerFilterBarConfig, USequencerSettings::FindTrackFilterBar
-	 */
-	explicit FSequencerFilterBar(FSequencer& InSequencer, FString InSubIdentifier);
+	explicit FSequencerFilterBar(UE::Sequencer::FFilterBarInitArgs InArgs);
 	virtual ~FSequencerFilterBar() override;
 
 	void BindCommands();
@@ -304,6 +296,9 @@ private:
 	 */
 	const FString SubIdentifier;
 	
+	/** Affects how filtering is performed. */
+	const UE::Sequencer::EFilterFlags Flags;
+	
 	/** @return Gets the selected IOutlinerExtensions, or all of them if nothing is selected. */
 	TSet<UE::Sequencer::TWeakViewModelPtr<UE::Sequencer::IOutlinerExtension>> GetSelectedTracksOrAll() const;
 	
diff --git a/Engine/Source/Editor/Sequencer/Private/Filters/ViewModels/Linking/LinkedFilterFactoryViewModel.cpp b/Engine/Source/Editor/Sequencer/Private/Filters/ViewModels/Linking/LinkedFilterFactoryViewModel.cpp
--- a/Engine/Source/Editor/Sequencer/Private/Filters/ViewModels/Linking/LinkedFilterFactoryViewModel.cpp
+++ b/Engine/Source/Editor/Sequencer/Private/Filters/ViewModels/Linking/LinkedFilterFactoryViewModel.cpp
@@ -20,7 +20,7 @@ FLinkedFilterFactoryViewModel::FLinkedFilterFactoryViewModel(
 TSharedRef<FLinkedFilterViewModel> FLinkedFilterFactoryViewModel::MakeFilteringModelImpl(const FMakeLinkedFilterViewModelArgs& InArgs) const
 {
 	const TSharedRef<FSequencerFilterBar> InstancedFilterBar = MakeShared<FSequencerFilterBar>(
-		Sequencer, InArgs.ConfigSubIdentifier
+		FFilterBarInitArgs(Sequencer, InArgs.ConfigSubIdentifier, InArgs.UnlinkedFilterFlags)
 		);
 	InstancedFilterBar->BindCommands();
 	return MakeShared<FLinkedFilterViewModel>(SharedFilterBar, InstancedFilterBar, FilterEvaluator);
diff --git a/Engine/Source/Editor/Sequencer/Private/Filters/ViewModels/Linking/MakeLinkedFilterViewModelArgs.h b/Engine/Source/Editor/Sequencer/Private/Filters/ViewModels/Linking/MakeLinkedFilterViewModelArgs.h
--- a/Engine/Source/Editor/Sequencer/Private/Filters/ViewModels/Linking/MakeLinkedFilterViewModelArgs.h
+++ b/Engine/Source/Editor/Sequencer/Private/Filters/ViewModels/Linking/MakeLinkedFilterViewModelArgs.h
@@ -3,6 +3,7 @@
 #pragma once
 
 #include "Containers/UnrealString.h"
+#include "Filters/FilterBarFlags.h"
 
 namespace UE::Sequencer
 {
@@ -19,8 +20,17 @@ struct FMakeLinkedFilterViewModelArgs
 	 */
 	FString ConfigSubIdentifier;
 	
-	explicit FMakeLinkedFilterViewModelArgs(FString InConfigSubIdentifier)
+	/** 
+	 * The filter flags to use for the unlinked filter bar.
+	 * 
+	 * By default, unlinked filter state should not be affected by the USequencerSettings::bIncludePinnedInFilter setting, so you likely don't want 
+	 * to set PinnedItemsCanSkipFiltering. 
+	 */
+	EFilterFlags UnlinkedFilterFlags = EFilterFlags::None;
+	
+	explicit FMakeLinkedFilterViewModelArgs(FString InConfigSubIdentifier, EFilterFlags InFilterFlags = EFilterFlags::None)
 		: ConfigSubIdentifier(MoveTemp(InConfigSubIdentifier))
+		, UnlinkedFilterFlags(InFilterFlags)
 		{}
 };
 } // namespace UE::Sequencer
diff --git a/Engine/Source/Editor/Sequencer/Private/MVVM/ViewModels/SequencerEditorViewModelImpl.cpp b/Engine/Source/Editor/Sequencer/Private/MVVM/ViewModels/SequencerEditorViewModelImpl.cpp
--- a/Engine/Source/Editor/Sequencer/Private/MVVM/ViewModels/SequencerEditorViewModelImpl.cpp
+++ b/Engine/Source/Editor/Sequencer/Private/MVVM/ViewModels/SequencerEditorViewModelImpl.cpp
@@ -34,14 +34,19 @@ void FSequencerEditorViewModelImpl::InitLinkedFiltering()
 	// For now, filtering will default to linked mode. For linked mode, we specify no sub-identifier: that causes FSequencerFilterBar to continue 
 	// using the old config entry, which would be SequencerSettings->GetName(), e.g. "LevelSequencerEditor".
 	const TCHAR* ConfigId = ConfigIds::LinkedConfigSubId;
+	// Linked filter state should be affected by the USequencerSettings::bIncludePinnedInFilter setting.
+	constexpr EFilterFlags LinkedFilterFlags = EFilterFlags::PinnedItemsCanSkipFiltering;
+	
 	LinkedFilterFactory = MakeShared<FLinkedFilterFactoryViewModel>(
 		*Sequencer,
-		MakeShared<FSequencerFilterBar>(*Sequencer, ConfigId),
+		MakeShared<FSequencerFilterBar>(FFilterBarInitArgs(*Sequencer, ConfigId, LinkedFilterFlags)),
 		Sequencer->GetNodeTree()->GetFilterEvaluator()
 		);
 	
+	// The unlinked filter state for the Outliner should also be affected by the USequencerSettings::bIncludePinnedInFilter setting.
+	constexpr EFilterFlags UnlinkedFilterFlags = EFilterFlags::PinnedItemsCanSkipFiltering;
 	SequenceFilterInstanceViewModel = LinkedFilterFactory->MakeFilteringModelImpl(
-		FMakeLinkedFilterViewModelArgs(ConfigIds::Sequencer_InstancedConfigSubId)
+		FMakeLinkedFilterViewModelArgs(ConfigIds::Sequencer_InstancedConfigSubId, UnlinkedFilterFlags)
 		);
 	FilterMediator = MakeShared<FLinkedFilteringMediator>(Sequencer, SequenceFilterInstanceViewModel.ToSharedRef());
 }
