Sequencer: Move files out from Filters/ViewModels.

Some of the files moved:
- Filters/ViewModels/Linking to Filters/Linking
- Filters/ViewModels to Filters/Evaluation

#Synced-CL 51626380

diff --git a/Engine/Source/Editor/Sequencer/Private/CurveEditor/Integration/CurveEditorCommandBinder.cpp b/Engine/Source/Editor/Sequencer/Private/CurveEditor/Integration/CurveEditorCommandBinder.cpp
--- a/Engine/Source/Editor/Sequencer/Private/CurveEditor/Integration/CurveEditorCommandBinder.cpp
+++ b/Engine/Source/Editor/Sequencer/Private/CurveEditor/Integration/CurveEditorCommandBinder.cpp
@@ -4,8 +4,8 @@
 
 #include "CurveEditor.h"
 #include "Filters/SequencerTrackFilterCommands.h"
-#include "Filters/ViewModels/Linking/FilterAreaManager.h"
-#include "Filters/ViewModels/Linking/LinkedFilterViewModel.h"
+#include "Filters/Linking/FilterAreaManager.h"
+#include "Filters/Linking/Mode/LinkedFilterViewModel.h"
 #include "ISequencer.h"
 #include "Sequencer.h"
 #include "SequencerCommands.h"
diff --git a/Engine/Source/Editor/Sequencer/Private/CurveEditor/Integration/CurveEditorCommandBinder.h b/Engine/Source/Editor/Sequencer/Private/CurveEditor/Integration/CurveEditorCommandBinder.h
--- a/Engine/Source/Editor/Sequencer/Private/CurveEditor/Integration/CurveEditorCommandBinder.h
+++ b/Engine/Source/Editor/Sequencer/Private/CurveEditor/Integration/CurveEditorCommandBinder.h
@@ -2,7 +2,7 @@
 
 #pragma once
 
-#include "Filters/ViewModels/Linking/FilterAreaManager.h"
+#include "Filters/Linking/FilterAreaManager.h"
 #include "Templates/SharedPointer.h"
 #include "Templates/UnrealTemplate.h"
 
diff --git a/Engine/Source/Editor/Sequencer/Private/CurveEditor/Integration/Selection/CurveModelSyncer.cpp b/Engine/Source/Editor/Sequencer/Private/CurveEditor/Integration/Selection/CurveModelSyncer.cpp
--- a/Engine/Source/Editor/Sequencer/Private/CurveEditor/Integration/Selection/CurveModelSyncer.cpp
+++ b/Engine/Source/Editor/Sequencer/Private/CurveEditor/Integration/Selection/CurveModelSyncer.cpp
@@ -4,7 +4,7 @@
 
 #include "CurveEditor.h"
 #include "CurveEditor/Views/CurveEditorWidgetOwner.h"
-#include "Filters/ViewModels/Linking/ILinkedFilterViewModel.h"
+#include "Filters/Linking/Mode/ILinkedFilterViewModel.h"
 #include "MVVM/Extensions/IOutlinerExtension.h"
 #include "MVVM/Selection/Selection.h"
 #include "MVVM/SharedViewModelData.h"
diff --git a/Engine/Source/Editor/Sequencer/Private/CurveEditor/SequencerCurveEditorApp.cpp b/Engine/Source/Editor/Sequencer/Private/CurveEditor/SequencerCurveEditorApp.cpp
--- a/Engine/Source/Editor/Sequencer/Private/CurveEditor/SequencerCurveEditorApp.cpp
+++ b/Engine/Source/Editor/Sequencer/Private/CurveEditor/SequencerCurveEditorApp.cpp
@@ -3,10 +3,10 @@
 #include "SequencerCurveEditorApp.h"
 
 #include "Filters/FilterConfigIdentifiers.h"
-#include "Filters/ViewModels/Linking/FilterAreaManager.h"
-#include "Filters/ViewModels/Linking/LinkedFilterFactoryViewModel.h"
-#include "Filters/ViewModels/Linking/LinkedFilterViewModel.h"
-#include "Filters/ViewModels/Linking/MakeLinkedFilterViewModelArgs.h"
+#include "Filters/Linking/FilterAreaManager.h"
+#include "Filters/Linking/Mode/LinkedFilterFactoryViewModel.h"
+#include "Filters/Linking/Mode/LinkedFilterViewModel.h"
+#include "Filters/Linking/Mode/MakeLinkedFilterViewModelArgs.h"
 #include "ICurveEditorModule.h"
 #include "MVVM/CurveEditorExtension.h"
 #include "MVVM/ViewModels/SequencerEditorViewModelImpl.h"
diff --git a/Engine/Source/Editor/Sequencer/Private/CurveEditor/SequencerCurveEditorApp.h b/Engine/Source/Editor/Sequencer/Private/CurveEditor/SequencerCurveEditorApp.h
--- a/Engine/Source/Editor/Sequencer/Private/CurveEditor/SequencerCurveEditorApp.h
+++ b/Engine/Source/Editor/Sequencer/Private/CurveEditor/SequencerCurveEditorApp.h
@@ -3,7 +3,7 @@
 #pragma once
 
 #include "CurveEditor/Views/CurveEditorWidgetOwner.h"
-#include "Filters/ViewModels/Linking/FilterAreaManager.h"
+#include "Filters/Linking/FilterAreaManager.h"
 #include "Integration/CurveEditorCommandBinder.h"
 #include "Integration/Selection/SelectionSyncRootLogic.h"
 #include "Templates/SharedPointer.h"
diff --git a/Engine/Source/Editor/Sequencer/Private/CurveEditor/Views/CurveEditorWidgetOwner.cpp b/Engine/Source/Editor/Sequencer/Private/CurveEditor/Views/CurveEditorWidgetOwner.cpp
--- a/Engine/Source/Editor/Sequencer/Private/CurveEditor/Views/CurveEditorWidgetOwner.cpp
+++ b/Engine/Source/Editor/Sequencer/Private/CurveEditor/Views/CurveEditorWidgetOwner.cpp
@@ -7,7 +7,7 @@
 #include "CurveEditor/Widgets/Tree/SCurveTreeContent.h"
 #include "CurveEditor/Widgets/Tree/SOldCurveTreeContent.h"
 #include "Filters/SCurveEditorFilterPanel.h"
-#include "Filters/ViewModels/Linking/LinkedFilterViewModel.h"
+#include "Filters/Linking/Mode/LinkedFilterViewModel.h"
 #include "Filters/Widgets/Linking/SFilterBarSwitcher.h"
 #include "HAL/IConsoleManager.h"
 #include "IPropertyRowGenerator.h"
diff --git a/Engine/Source/Editor/Sequencer/Private/CurveEditor/Widgets/SCurveEditorSearchAndFilterRow.cpp b/Engine/Source/Editor/Sequencer/Private/CurveEditor/Widgets/SCurveEditorSearchAndFilterRow.cpp
--- a/Engine/Source/Editor/Sequencer/Private/CurveEditor/Widgets/SCurveEditorSearchAndFilterRow.cpp
+++ b/Engine/Source/Editor/Sequencer/Private/CurveEditor/Widgets/SCurveEditorSearchAndFilterRow.cpp
@@ -8,8 +8,8 @@
 #include "Filters/Utils/HideIsolateViewModelUtils.h"
 #include "Filters/Utils/LinkedFilteringViewUtils.h"
 #include "Filters/ViewModels/HideIsolateShowViewModel.h"
-#include "Filters/ViewModels/Linking/ILinkedFilterViewModel.h"
-#include "Filters/ViewModels/Linking/LinkedFilterViewModel.h"
+#include "Filters/Linking/Mode/ILinkedFilterViewModel.h"
+#include "Filters/Linking/Mode/LinkedFilterViewModel.h"
 #include "Filters/Widgets/Linking/SFilterLinkStateSwitcher.h"
 #include "Filters/Widgets/Linking/SLinkedFilteringToggleButton.h"
 #include "Filters/Widgets/SSearchAndFilterWidget.h"
diff --git a/Engine/Source/Editor/Sequencer/Private/Filters/Evaluation/FilterEvaluator.cpp b/Engine/Source/Editor/Sequencer/Private/Filters/Evaluation/FilterEvaluator.cpp
new file mode 100644
--- /dev/null
+++ b/Engine/Source/Editor/Sequencer/Private/Filters/Evaluation/FilterEvaluator.cpp
@@ -0,0 +1,130 @@
+﻿// Copyright Epic Games, Inc. All Rights Reserved.
+
+#include "FilterEvaluator.h"
+
+#include "Filters/SequencerFilterBar.h"
+
+namespace UE::Sequencer
+{
+bool FFilterEvaluator::UpdateFiltersOnTrackValueChanged()
+{
+	if (bHasRequestedFilterUpdate)
+	{
+		return false;
+	}
+	
+	bool bAnyFilterNeedsRerun = false;
+	for (FFilterBarData& FilterBarData : RegisteredFilterBars)
+	{
+		const bool bNeedsRerun = FilterBarData.FilterBar->ShouldUpdateOnTrackValueChanged();
+		FilterBarData.bHasRequestedRefilter |= bNeedsRerun;
+		bAnyFilterNeedsRerun |= bNeedsRerun;
+	}
+	
+	if (bAnyFilterNeedsRerun)
+	{
+		constexpr bool bForce = true;
+		const bool bFiltersUpdated = UpdateFilters(bForce);
+
+		// If the filter list was modified, set bHasRequestedFilterUpdate to suppress excessive re-filters between tree update
+		bHasRequestedFilterUpdate = bFiltersUpdated;
+		return bFiltersUpdated;
+	}
+
+	return false;
+}
+
+bool FFilterEvaluator::UpdateFilters(bool bForce)
+{
+	if (!bForce && !bHasRequestedFilterUpdate)
+	{
+		return false;
+	}
+
+	bool bHasFilterDataChanged = false; 
+	for (FFilterBarData& FilterBarData : RegisteredFilterBars)
+	{
+		const TSharedRef<FSequencerFilterBar>& FilterBar = FilterBarData.FilterBar;
+		
+		const FSequencerFilterData PreviousFilterData = FilterBar->GetFilterData();
+		const FSequencerFilterData& FilterData = FilterBar->FilterNodes();
+		const bool bHasFilteredDataChanged = FilterData != PreviousFilterData;
+		bHasFilterDataChanged |= bHasFilteredDataChanged;
+		FilterBarData.bHasRequestedRefilter = false;
+		
+		OnFilterBarFilteredDelegate.Broadcast(FilterBar);
+	}
+	bHasRequestedFilterUpdate = false;
+
+	return bHasFilterDataChanged;
+}
+
+void FFilterEvaluator::RequestFilterUpdate(const TSharedRef<FSequencerFilterBar>& InFilterBar)
+{
+	const int32 Index = IndexOf(InFilterBar);
+	if (RegisteredFilterBars.IsValidIndex(Index))
+	{
+		bHasRequestedFilterUpdate = true;
+		RegisteredFilterBars[Index].bHasRequestedRefilter = true;
+	}
+}
+
+bool FFilterEvaluator::NeedsFilterUpdate() const
+{
+	return bHasRequestedFilterUpdate;
+}
+
+void FFilterEvaluator::IncrementFilterBarUsage(const TSharedRef<FSequencerFilterBar>& InFilterBar)
+{
+	const int32 Index = IndexOf(InFilterBar);
+	if (RegisteredFilterBars.IsValidIndex(Index))
+	{
+		++RegisteredFilterBars[Index].UsageCount;
+	}
+	else
+	{
+		FSimpleDelegate HandleRequestUpdate = FSimpleDelegate::CreateRaw(this, &FFilterEvaluator::OnFilterRequestedRefresh, InFilterBar.ToWeakPtr());
+		RegisteredFilterBars.Emplace(InFilterBar, MoveTemp(HandleRequestUpdate));
+	}
+	
+	// We run the filter just in case it's not been evaluated yet. In the future, we could add a dirty flag or serial number to the filter bar.
+	bHasRequestedFilterUpdate = true;
+}
+
+void FFilterEvaluator::DecrementFilterBarUsage(const TSharedRef<FSequencerFilterBar>& InFilterBar)
+{
+	const int32 Index = IndexOf(InFilterBar);
+	if (!ensure(RegisteredFilterBars.IsValidIndex(Index)))
+	{
+		return;
+	}
+	
+	FFilterBarData& Data = RegisteredFilterBars[Index];
+	check(Data.UsageCount > 0);
+	
+	if (Data.UsageCount == 1)
+	{
+		RegisteredFilterBars.RemoveAtSwap(Index);
+	}
+	else
+	{
+		--Data.UsageCount;
+	}
+}
+
+int32 FFilterEvaluator::IndexOf(const TSharedRef<FSequencerFilterBar>& InFilterBar) const
+{
+	return RegisteredFilterBars.IndexOfByPredicate([&InFilterBar](const FFilterBarData& Data)
+	{
+		return InFilterBar == Data.FilterBar;
+	});
+}
+
+void FFilterEvaluator::OnFilterRequestedRefresh(TWeakPtr<FSequencerFilterBar> InWeakFilterBar)
+{
+	if (const TSharedPtr<FSequencerFilterBar> FilterBar = InWeakFilterBar.Pin(); ensure(FilterBar))
+	{
+		RequestFilterUpdate(FilterBar.ToSharedRef());
+	}
+}
+} // namespace UE::Sequencer
diff --git a/Engine/Source/Editor/Sequencer/Private/Filters/Evaluation/FilterEvaluator.h b/Engine/Source/Editor/Sequencer/Private/Filters/Evaluation/FilterEvaluator.h
new file mode 100644
--- /dev/null
+++ b/Engine/Source/Editor/Sequencer/Private/Filters/Evaluation/FilterEvaluator.h
@@ -0,0 +1,106 @@
+﻿// Copyright Epic Games, Inc. All Rights Reserved.
+
+#pragma once
+
+#include "Containers/Array.h"
+#include "Delegates/Delegate.h"
+#include "Filters/SequencerFilterBar.h"
+#include "Templates/SharedPointer.h"
+#include "Templates/UnrealTemplate.h"
+
+class FSequencerFilterBar;
+
+namespace UE::Sequencer
+{
+/**
+ * Knows of all filter bar instances that need to be run in Sequencer.
+ * This is relevant for multiple UI endpoints using the linked filtering system.
+ * 
+ * Each endpoint can register a filter bar they want evaluated, and request explicit refreshes.
+ * The evaluator will be refreshed next tick by Sequencer.
+ * @see IncrementFilterBarUsage, DecrementFilterBarUsage
+ * 
+ * If your filter UI becomes hidden (e.g. tab in background), as optimization, it would be valid to call DecrementFilterBarUsage and follow it up 
+ * again with a IncrementFilterBarUsage once the UI becomes active again.
+ */
+class FFilterEvaluator : public FNoncopyable
+{
+public:
+	
+	/** 
+	 * Evaluates all filters bars that need updating when a track value changes. 
+	 * @return Whether the new list of filtered nodes is different from the previous list
+	 */
+	bool UpdateFiltersOnTrackValueChanged();
+	
+	/**
+	 * Evaluates all filter bars. 
+	 * @return Whether the new list of filtered nodes is different from the previous list
+	 */
+	bool UpdateFilters(bool bForce = false);
+	
+	/** Schedules an update of all filters on the given filter bar. The filter bar must be registered. */
+	void RequestFilterUpdate(const TSharedRef<FSequencerFilterBar>& InFilterBar);
+	/** @return Whether there is a filter update scheduled */
+	bool NeedsFilterUpdate() const;
+
+	/** 
+	 * Tracks a filter bar that some UI endpoints is displaying filter results for.
+	 * 
+	 * If this is already tracked, increments the usage counter by 1.
+	 * If previously untracked, this creates a subscription ISequencerFilterBar::OnRequestUpdate.
+	 * 
+	 * @note You should prefer FScopedFilterBarUsage for auto deregistration
+	 * @see ScopedFilterBarUsage.
+	 * 
+	 * Each call must be matched with DecrementFilterBarUsage.
+	 * @see DecrementFilterBarUsage
+	 */
+	void IncrementFilterBarUsage(const TSharedRef<FSequencerFilterBar>& InFilterBar);
+	/** Decrements the usage counter by 1. If it reaches 0, the filter bar becomes untracked for filter evaluation. */
+	void DecrementFilterBarUsage(const TSharedRef<FSequencerFilterBar>& InFilterBar);
+	
+	DECLARE_MULTICAST_DELEGATE_OneParam(FOnFilterBarFiltered, const TSharedRef<FSequencerFilterBar>&)
+	/** @return Delegate invoked when a filter bar is refiltered. */
+	FOnFilterBarFiltered& OnFilterBarFiltered() { return OnFilterBarFilteredDelegate; }
+	
+private:
+	
+	/** Whether any of the filters has a refilter requested. */
+	bool bHasRequestedFilterUpdate = false;
+	
+	struct FFilterBarData : FNoncopyable // Avoid accidental copying so we don't unsubscribe by accident.
+	{
+		/** The filter bar this data corresponds to. */
+		const TSharedRef<FSequencerFilterBar> FilterBar;
+		/** Handle to ISequencerFilterBar::OnRequestUpdate. */
+		const FDelegateHandle OnRequestUpdateHandle;
+		
+		/** How many systems have registered to be using this filter bar. Once this reaches 0, the entry is removed. */
+		int32 UsageCount = 1;
+		/** Whether this filter bar has been marked to be refiltered. */
+		bool bHasRequestedRefilter = false;
+
+		explicit FFilterBarData(const TSharedRef<FSequencerFilterBar>& InFilterBar, FSimpleDelegate HandleFilterRequestRefresh)
+			: FilterBar(InFilterBar)
+			, OnRequestUpdateHandle(FilterBar->OnRequestUpdate().Add(MoveTemp(HandleFilterRequestRefresh)))
+		{}
+		
+		~FFilterBarData()
+		{
+			FilterBar->OnRequestUpdate().Remove(OnRequestUpdateHandle);
+		}
+	};
+	/** Filters that should be considered for evaluation. */
+	TArray<FFilterBarData> RegisteredFilterBars;
+	
+	/** Invoked when a filter bar is refiltered. */
+	FOnFilterBarFiltered OnFilterBarFilteredDelegate;
+	
+	/** @return Index to RegisteredFilterBars entry that contains InFilterBar. */
+	int32 IndexOf(const TSharedRef<FSequencerFilterBar>& InFilterBar) const;
+	
+	/** Handles ISequencerTrackFilters::OnRequestRefresh by calling RequestFilterUpdate. */
+	void OnFilterRequestedRefresh(TWeakPtr<FSequencerFilterBar> InWeakFilterBar);
+};
+} // namespace UE::Sequencer
diff --git a/Engine/Source/Editor/Sequencer/Private/Filters/Evaluation/ScopedFilterBarUsage.h b/Engine/Source/Editor/Sequencer/Private/Filters/Evaluation/ScopedFilterBarUsage.h
new file mode 100644
--- /dev/null
+++ b/Engine/Source/Editor/Sequencer/Private/Filters/Evaluation/ScopedFilterBarUsage.h
@@ -0,0 +1,69 @@
+﻿// Copyright Epic Games, Inc. All Rights Reserved.
+
+#pragma once
+
+#include "FilterEvaluator.h"
+#include "Templates/SharedPointer.h"
+#include "Templates/UnrealTemplate.h"
+
+class FSequencerFilterBar;
+
+namespace UE::Sequencer
+{
+class FFilterEvaluator;
+
+/**
+ * Registers FSequencerFilterBar with FFilterEvaluator for the lifetime of this scope object.
+ * Handles calling FFilterEvaluator::IncrementFilterBarUsage and FFilterEvaluator::DecrementFilterBarUsage.
+ */
+class FScopedFilterBarUsage : public FNoncopyable
+{
+public:
+	
+	explicit FScopedFilterBarUsage(const TSharedRef<FSequencerFilterBar> InFilterBar, const TSharedRef<FFilterEvaluator>& InEvaluator)
+		: FilterBar(InFilterBar)
+		, Evaluator(InEvaluator)
+	{
+		Evaluator->IncrementFilterBarUsage(InFilterBar);
+	}
+	
+	FScopedFilterBarUsage(FScopedFilterBarUsage&& InOther)
+		: FilterBar(MoveTemp(InOther.FilterBar))
+		, Evaluator(MoveTemp(InOther.Evaluator))
+	{}
+
+	FScopedFilterBarUsage& operator=(FScopedFilterBarUsage&& Other)
+	{
+		if (this != &Other)
+		{
+			CleanUp();
+		}
+		
+		FilterBar = MoveTemp(Other.FilterBar);
+		Evaluator = MoveTemp(Other.Evaluator);
+		return *this;
+	}
+
+	~FScopedFilterBarUsage()
+	{
+		CleanUp();
+	}
+	
+private:
+	
+	/** The filter bar that is being used. */
+	TSharedPtr<FSequencerFilterBar> FilterBar;
+	
+	/** The filter evaluator to register and unregister FilterBar with. */
+	TSharedPtr<FFilterEvaluator> Evaluator;
+	
+	/** Calls DecrementFilterBarUsage if this instance state has not been moved. */
+	void CleanUp() const
+	{
+		if (Evaluator && FilterBar)
+		{
+			Evaluator->DecrementFilterBarUsage(FilterBar.ToSharedRef());
+		}
+	}
+};
+} // namespace UE::Sequencer
diff --git a/Engine/Source/Editor/Sequencer/Private/Filters/Linking/FilterAreaManager.cpp b/Engine/Source/Editor/Sequencer/Private/Filters/Linking/FilterAreaManager.cpp
new file mode 100644
--- /dev/null
+++ b/Engine/Source/Editor/Sequencer/Private/Filters/Linking/FilterAreaManager.cpp
@@ -0,0 +1,54 @@
+﻿// Copyright Epic Games, Inc. All Rights Reserved.
+
+#include "FilterAreaManager.h"
+
+#include "Sequencer.h"
+#include "Filters/SequencerTrackFilterCommands.h"
+
+namespace UE::Sequencer
+{
+FFilterAreaManager::FFilterAreaManager(
+	FName InFilterAreaConfigId,
+	TWeakPtr<FSequencer> InWeakSequencer, 
+	const TSharedRef<FLinkedFilterViewModel>& InFilterModel
+	)
+	: FilterAreaConfigId(InFilterAreaConfigId)
+	, WeakSequencer(MoveTemp(InWeakSequencer))
+	, FilterModel(InFilterModel)
+{
+	BindCommands();
+}
+
+void FFilterAreaManager::SetFilterBarVisible(bool bIsVisible)
+{
+	const TSharedPtr<FSequencer> SequencerPin = WeakSequencer.Pin();
+	USequencerSettings* Settings = SequencerPin ? SequencerPin->GetSequencerSettings() : nullptr;
+	if (Settings)
+	{
+		constexpr bool bSaveConfig = false; // Saving unnecessary because we'll save right after settings
+		Settings->FindOrAddTrackFilterArea(FilterAreaConfigId, bSaveConfig).SetIsFilterBarVisible(bIsVisible);
+		Settings->SaveConfig();
+	}
+}
+
+bool FFilterAreaManager::IsFilterBarVisible() const
+{
+	const TSharedPtr<FSequencer> SequencerPin = WeakSequencer.Pin();
+	USequencerSettings* Settings = SequencerPin ? SequencerPin->GetSequencerSettings() : nullptr;
+	
+	constexpr bool bSaveConfig = false; // Saving unnecessary because we're just getting info and using the default if not found.
+	return Settings && Settings->FindOrAddTrackFilterArea(FilterAreaConfigId, bSaveConfig).IsFilterBarVisible();
+}
+
+void FFilterAreaManager::BindCommands()
+{
+	FSequencerTrackFilterCommands& FilterCommands = FSequencerTrackFilterCommands::Get();
+	
+	CommandList->MapAction(
+		FilterCommands.ToggleFilterBarVisibility,
+		FExecuteAction::CreateRaw(this, &FFilterAreaManager::ToggleFilterBarVisibility),
+		FCanExecuteAction::CreateLambda([]{ return true; }),
+		FIsActionChecked::CreateRaw(this, &FFilterAreaManager::IsFilterBarVisible)
+		);
+}
+}
diff --git a/Engine/Source/Editor/Sequencer/Private/Filters/Linking/FilterAreaManager.h b/Engine/Source/Editor/Sequencer/Private/Filters/Linking/FilterAreaManager.h
new file mode 100644
--- /dev/null
+++ b/Engine/Source/Editor/Sequencer/Private/Filters/Linking/FilterAreaManager.h
@@ -0,0 +1,57 @@
+﻿// Copyright Epic Games, Inc. All Rights Reserved.
+
+#pragma once
+
+#include "Framework/Commands/UICommandList.h"
+#include "Templates/SharedPointer.h"
+#include "Templates/UnrealTemplate.h"
+
+class FSequencer;
+class FUICommandList;
+
+namespace UE::Sequencer
+{
+class FLinkedFilterViewModel;
+
+/** 
+ * Manages a UI area that displays filters that can be linked or instanced.
+ * It manages visibility of the filter pills and handles the related commands.
+ */
+class FFilterAreaManager : public FNoncopyable
+{
+public:
+	
+	explicit FFilterAreaManager(
+		FName InFilterAreaConfigId,
+		TWeakPtr<FSequencer> InWeakSequencer, 
+		const TSharedRef<FLinkedFilterViewModel>& InFilterModel
+		);
+	
+	/** Sets whether the filter bar pills should be visible. */
+	void SetFilterBarVisible(bool bIsVisible);
+	/** @return Whether the filter bar pills should be visible */
+	bool IsFilterBarVisible() const;
+	/** Inverts the filter bar pill visibility. */
+	void ToggleFilterBarVisibility() { SetFilterBarVisible(!IsFilterBarVisible()); }
+	
+	const TSharedRef<FLinkedFilterViewModel>& GetFilterModel() const { return FilterModel; }
+	const TSharedRef<FUICommandList>& GetFilterAreaCommandList() const { return CommandList; }
+	
+private:
+	
+	/** @see USequencerSettings::FindOrAddTrackFilterArea. */
+	const FName FilterAreaConfigId;
+	
+	/** Used to get the settings, e.g. to save filter bar visbility. */
+	const TWeakPtr<FSequencer> WeakSequencer;
+	
+	/** Manages the linked and unlinked filter state. */
+	const TSharedRef<FLinkedFilterViewModel> FilterModel;
+	
+	/** Handles the user's commands, e.g. to change visibility. */
+	const TSharedRef<FUICommandList> CommandList = MakeShared<FUICommandList>();
+	
+	void BindCommands();
+};
+}
+
diff --git a/Engine/Source/Editor/Sequencer/Private/Filters/Linking/Mode/ILinkedFilterFactoryViewModel.h b/Engine/Source/Editor/Sequencer/Private/Filters/Linking/Mode/ILinkedFilterFactoryViewModel.h
new file mode 100644
--- /dev/null
+++ b/Engine/Source/Editor/Sequencer/Private/Filters/Linking/Mode/ILinkedFilterFactoryViewModel.h
@@ -0,0 +1,29 @@
+﻿// Copyright Epic Games, Inc. All Rights Reserved.
+
+#pragma once
+
+#include "Templates/SharedPointerFwd.h"
+
+namespace UE::Sequencer
+{
+class ILinkedFilterViewModel;
+struct FMakeLinkedFilterViewModelArgs;
+
+/** 
+ * Factory for creating ILinkedFilterViewModel.
+ * Use this if you want to create UI that can be linked to the filters in the Sequencer UI.
+ */
+class ILinkedFilterFactoryViewModel
+{
+public:
+	
+	/**
+	 * Creates a new view model instance for UI that wants to be able to link to the filter state in Sequencer.
+	 * The view model has a reference to the shared linked filter state and its own filter state, which you can switch between.
+	 * @return View model that manages linked and instanced filter state for your UI.
+	 */
+	virtual TSharedRef<ILinkedFilterViewModel> MakeFilteringModel(const FMakeLinkedFilterViewModelArgs& InArgs) = 0;
+	
+	virtual ~ILinkedFilterFactoryViewModel() = default;
+};
+} // namespace UE::Sequencer
diff --git a/Engine/Source/Editor/Sequencer/Private/Filters/Linking/Mode/ILinkedFilterViewModel.cpp b/Engine/Source/Editor/Sequencer/Private/Filters/Linking/Mode/ILinkedFilterViewModel.cpp
new file mode 100644
--- /dev/null
+++ b/Engine/Source/Editor/Sequencer/Private/Filters/Linking/Mode/ILinkedFilterViewModel.cpp
@@ -0,0 +1,14 @@
+﻿// Copyright Epic Games, Inc. All Rights Reserved.
+
+#include "ILinkedFilterViewModel.h"
+
+#include "Filters/ISequencerTrackFilters.h"
+#include "Filters/SequencerFilterData.h"
+
+namespace UE::Sequencer
+{
+bool ILinkedFilterViewModel::IsFilteredOut(const TViewModelPtr<IOutlinerExtension>& InItem) const
+{
+	return GetActiveFilterBar()->GetFilterData().IsFilteredOut(InItem);
+}
+} // namespace UE::Sequencer
diff --git a/Engine/Source/Editor/Sequencer/Private/Filters/Linking/Mode/ILinkedFilterViewModel.h b/Engine/Source/Editor/Sequencer/Private/Filters/Linking/Mode/ILinkedFilterViewModel.h
new file mode 100644
--- /dev/null
+++ b/Engine/Source/Editor/Sequencer/Private/Filters/Linking/Mode/ILinkedFilterViewModel.h
@@ -0,0 +1,58 @@
+﻿// Copyright Epic Games, Inc. All Rights Reserved.
+
+#pragma once
+
+#include "Delegates/Delegate.h"
+#include "HAL/Platform.h"
+#include "MVVM/ViewModelPtr.h"
+
+class ISequencerTrackFilters;
+
+namespace UE::Sequencer
+{
+class IOutlinerExtension;
+
+enum class ELinkedFilterMode : uint8
+{
+	/** Filtering is using a shared filterinstance */
+	Linked,
+	/** Filtering is using its own filter state. */
+	Instanced,
+	
+	/** Not a real mode. Useful for static asserting. */
+	Count
+};
+
+/** 
+ * View-model for UI to switch between:
+ * - shared filter state: this state is shared by multiple UI, allowing the user to sync multiple UI instances (Sequencer, Curve Editor, etc.)
+ * - instanced filter state: this state is used only by the UI instance
+ */
+class ILinkedFilterViewModel
+{
+public:
+	
+	/** Sets the current filtering mode. */
+	virtual void SetFilterMode(ELinkedFilterMode InNewMode) = 0;
+	/** @return The current filtering mode. */
+	virtual ELinkedFilterMode GetFilterMode() const = 0;
+	
+	/** @return Whether InItem is filtered out by the active filter bar. */
+	bool IsFilteredOut(const TViewModelPtr<IOutlinerExtension>& InItem) const;
+	
+	/** @return The currently active filter bar instance according to the linked mode. */
+	virtual TSharedRef<ISequencerTrackFilters> GetActiveFilterBar() const = 0; 
+	/** @return The filter bar that is shared with all other UI. */
+	virtual TSharedRef<ISequencerTrackFilters> GetLinkedFilterBar() const = 0;
+	/** @return The filter bar that is only used by this instance. */
+	virtual TSharedRef<ISequencerTrackFilters> GetInstancedFilterBar() const = 0;
+	
+	/** @return Event when the filter mode changes, i.e. the instance returned by GetActiveFilterBar changes. */
+	virtual FSimpleMulticastDelegate& OnFilterModeChanged() = 0;
+	
+	/** @return Event when the items displayed need to be refreshed, e.g. when ELinkedFilterMode changes, or the active filter bar refilters its items. */
+	virtual FSimpleMulticastDelegate& OnFilteredItemsChanged() = 0;
+	
+	virtual ~ILinkedFilterViewModel() = default;
+};
+} // namespace UE::Sequencer
\ No newline at end of file
diff --git a/Engine/Source/Editor/Sequencer/Private/Filters/Linking/Mode/LinkedFilterFactoryViewModel.cpp b/Engine/Source/Editor/Sequencer/Private/Filters/Linking/Mode/LinkedFilterFactoryViewModel.cpp
new file mode 100644
--- /dev/null
+++ b/Engine/Source/Editor/Sequencer/Private/Filters/Linking/Mode/LinkedFilterFactoryViewModel.cpp
@@ -0,0 +1,33 @@
+﻿// Copyright Epic Games, Inc. All Rights Reserved.
+
+#include "LinkedFilterFactoryViewModel.h"
+
+#include "LinkedFilterViewModel.h"
+#include "Filters/Linking/Mode/MakeLinkedFilterViewModelArgs.h"
+
+namespace UE::Sequencer
+{
+FLinkedFilterFactoryViewModel::FLinkedFilterFactoryViewModel(
+	FSequencer& InSequencer UE_LIFETIMEBOUND,
+	const TSharedRef<FSequencerFilterBar>& InSharedFilterBar,
+	const TSharedRef<FFilterEvaluator> InFilterEvaluator
+	)
+	: Sequencer(InSequencer)
+	, SharedFilterBar(InSharedFilterBar)
+	, FilterEvaluator(InFilterEvaluator)
+{}
+
+TSharedRef<FLinkedFilterViewModel> FLinkedFilterFactoryViewModel::MakeFilteringModelImpl(const FMakeLinkedFilterViewModelArgs& InArgs) const
+{
+	const TSharedRef<FSequencerFilterBar> InstancedFilterBar = MakeShared<FSequencerFilterBar>(
+		FFilterBarInitArgs(Sequencer, InArgs.ConfigSubIdentifier, InArgs.UnlinkedFilterFlags)
+		);
+	InstancedFilterBar->BindCommands();
+	return MakeShared<FLinkedFilterViewModel>(SharedFilterBar, InstancedFilterBar, FilterEvaluator);
+}
+
+TSharedRef<ILinkedFilterViewModel> FLinkedFilterFactoryViewModel::MakeFilteringModel(const FMakeLinkedFilterViewModelArgs& InArgs)
+{
+	return MakeFilteringModelImpl(InArgs);
+}
+} // namespace UE::Sequencer
diff --git a/Engine/Source/Editor/Sequencer/Private/Filters/Linking/Mode/LinkedFilterFactoryViewModel.h b/Engine/Source/Editor/Sequencer/Private/Filters/Linking/Mode/LinkedFilterFactoryViewModel.h
new file mode 100644
--- /dev/null
+++ b/Engine/Source/Editor/Sequencer/Private/Filters/Linking/Mode/LinkedFilterFactoryViewModel.h
@@ -0,0 +1,42 @@
+﻿// Copyright Epic Games, Inc. All Rights Reserved.
+
+#pragma once
+
+#include "Filters/SequencerFilterBar.h"
+#include "Filters/Linking/Mode/ILinkedFilterFactoryViewModel.h"
+#include "HAL/Platform.h"
+
+namespace UE::Sequencer
+{
+class FFilterEvaluator;
+class FLinkedFilterViewModel;
+
+class FLinkedFilterFactoryViewModel : public ILinkedFilterFactoryViewModel
+{
+public:
+	
+	explicit FLinkedFilterFactoryViewModel(
+		FSequencer& InSequencer UE_LIFETIMEBOUND,
+		const TSharedRef<FSequencerFilterBar>& InSharedFilterBar,
+		const TSharedRef<FFilterEvaluator> InFilterEvaluator
+		);
+	
+	/** @return View model that manages linked and instanced filter state for your UI.  */
+	TSharedRef<FLinkedFilterViewModel> MakeFilteringModelImpl(const FMakeLinkedFilterViewModelArgs& InArgs) const;
+	
+	//~ Begin ILinkedFilterFactoryViewModel
+	virtual TSharedRef<ILinkedFilterViewModel> MakeFilteringModel(const FMakeLinkedFilterViewModelArgs& InArgs) override;
+	//~ End ILinkedFilterFactoryViewModel
+	
+private:
+	
+	/** Sequencer */
+	FSequencer& Sequencer;
+	
+	/** The filter bar instance that model edit when they are in linked mode. All created models reference this. */
+	const TSharedRef<FSequencerFilterBar> SharedFilterBar;
+	
+	/** Instance that the active filter bar instances are registered with. This way Sequencer will issue a refilter eventually. */
+	const TSharedRef<FFilterEvaluator> FilterEvaluator;
+};
+} // namespace UE::Sequencer
diff --git a/Engine/Source/Editor/Sequencer/Private/Filters/Linking/Mode/LinkedFilterViewModel.cpp b/Engine/Source/Editor/Sequencer/Private/Filters/Linking/Mode/LinkedFilterViewModel.cpp
new file mode 100644
--- /dev/null
+++ b/Engine/Source/Editor/Sequencer/Private/Filters/Linking/Mode/LinkedFilterViewModel.cpp
@@ -0,0 +1,84 @@
+﻿// Copyright Epic Games, Inc. All Rights Reserved.
+
+#include "LinkedFilterViewModel.h"
+
+#include "Filters/SequencerFilterBar.h"
+#include "Filters/Evaluation/FilterEvaluator.h"
+
+namespace UE::Sequencer
+{
+FLinkedFilterViewModel::FLinkedFilterViewModel(
+	const TSharedRef<FSequencerFilterBar>& InSharedFilterBar,
+	const TSharedRef<FSequencerFilterBar>& InInstancedFilterBar,
+	const TSharedRef<FFilterEvaluator> InFilterEvaluator
+	)
+	: LinkedFilterBar(InSharedFilterBar)
+	, InstancedFilterBar(InInstancedFilterBar)
+	, FilterEvaluator(InFilterEvaluator)
+	, ActiveFilterBarUsage(GetActiveFilterBarImpl(), FilterEvaluator)
+{
+	FilterEvaluator->OnFilterBarFiltered().AddRaw(this, &FLinkedFilterViewModel::OnFilterBarFiltered);
+}
+
+FLinkedFilterViewModel::~FLinkedFilterViewModel()
+{
+	FilterEvaluator->OnFilterBarFiltered().RemoveAll(this);
+}
+
+TSharedRef<FSequencerFilterBar> FLinkedFilterViewModel::GetActiveFilterBarImpl() const
+{
+	static_assert(static_cast<int32>(ELinkedFilterMode::Count) == 2, "Update this switch");
+	switch (FilterMode)
+	{
+	case ELinkedFilterMode::Linked: return LinkedFilterBar;
+	case ELinkedFilterMode::Instanced: return InstancedFilterBar;
+	default: checkNoEntry(); return LinkedFilterBar;
+	}
+}
+
+TSharedRef<FSequencerFilterBar> FLinkedFilterViewModel::GetLinkedFilterBarImpl() const
+{
+	return LinkedFilterBar;
+}
+
+TSharedRef<FSequencerFilterBar> FLinkedFilterViewModel::GetInstancedFilterBarImpl() const
+{
+	return InstancedFilterBar;
+}
+
+void FLinkedFilterViewModel::SetFilterMode(ELinkedFilterMode InNewMode)
+{
+	if (InNewMode != FilterMode)
+	{
+		FilterMode = InNewMode;
+		
+		const TSharedRef<FSequencerFilterBar> NewActiveFilterBar = GetActiveFilterBarImpl();
+		ActiveFilterBarUsage = FScopedFilterBarUsage(NewActiveFilterBar, FilterEvaluator);
+		OnFilterModeChangedDelegate.Broadcast();
+		OnFilteredItemsChangedDelegate.Broadcast();
+	}
+}
+
+TSharedRef<ISequencerTrackFilters> FLinkedFilterViewModel::GetActiveFilterBar() const
+{
+	return GetActiveFilterBarImpl();
+}
+
+TSharedRef<ISequencerTrackFilters> FLinkedFilterViewModel::GetLinkedFilterBar() const
+{
+	return LinkedFilterBar;
+}
+
+TSharedRef<ISequencerTrackFilters> FLinkedFilterViewModel::GetInstancedFilterBar() const
+{
+	return InstancedFilterBar;
+}
+
+void FLinkedFilterViewModel::OnFilterBarFiltered(const TSharedRef<FSequencerFilterBar>& InFilterBar)
+{
+	if (InFilterBar == GetActiveFilterBarImpl())
+	{
+		OnFilteredItemsChangedDelegate.Broadcast();
+	}
+}
+} // namespace UE::Sequencer
diff --git a/Engine/Source/Editor/Sequencer/Private/Filters/Linking/Mode/LinkedFilterViewModel.h b/Engine/Source/Editor/Sequencer/Private/Filters/Linking/Mode/LinkedFilterViewModel.h
new file mode 100644
--- /dev/null
+++ b/Engine/Source/Editor/Sequencer/Private/Filters/Linking/Mode/LinkedFilterViewModel.h
@@ -0,0 +1,63 @@
+﻿// Copyright Epic Games, Inc. All Rights Reserved.
+
+#pragma once
+
+#include "Filters/Evaluation/ScopedFilterBarUsage.h"
+#include "Filters/Linking/Mode/ILinkedFilterViewModel.h"
+#include "Templates/SharedPointer.h"
+
+class FSequencer;
+class FSequencerFilterBar;
+
+namespace UE::Sequencer
+{
+class FFilterEvaluator;
+
+class FLinkedFilterViewModel : public ILinkedFilterViewModel
+{
+public:
+	
+	explicit FLinkedFilterViewModel(
+		const TSharedRef<FSequencerFilterBar>& InSharedFilterBar,
+		const TSharedRef<FSequencerFilterBar>& InInstancedFilterBar,
+		const TSharedRef<FFilterEvaluator> InFilterEvaluator
+		);
+	~FLinkedFilterViewModel();
+	
+	TSharedRef<FSequencerFilterBar> GetActiveFilterBarImpl() const;
+	TSharedRef<FSequencerFilterBar> GetLinkedFilterBarImpl() const;
+	TSharedRef<FSequencerFilterBar> GetInstancedFilterBarImpl() const;
+	
+	//~ Begin ILinkedFilterViewModel
+	virtual void SetFilterMode(ELinkedFilterMode InNewMode) override;
+	virtual ELinkedFilterMode GetFilterMode() const override { return FilterMode; }
+	virtual TSharedRef<ISequencerTrackFilters> GetActiveFilterBar() const override;
+	virtual TSharedRef<ISequencerTrackFilters> GetLinkedFilterBar() const override;
+	virtual TSharedRef<ISequencerTrackFilters> GetInstancedFilterBar() const override;
+	virtual FSimpleMulticastDelegate& OnFilterModeChanged() override { return OnFilterModeChangedDelegate; }
+	virtual FSimpleMulticastDelegate& OnFilteredItemsChanged() override { return OnFilteredItemsChangedDelegate; }
+	//~ End ILinkedFilterViewModel
+	
+private:
+	
+	/** The filter bar instance that is shared across all UI. Other UI can choose to link to this instance. */
+	const TSharedRef<FSequencerFilterBar> LinkedFilterBar;
+	/** The filter bar instance that is owned by the particular view that is displaying the filters, e.g. Curve Editor. */
+	const TSharedRef<FSequencerFilterBar> InstancedFilterBar;
+	
+	/** The current filter mode. It determines which filter bar instance is used. */
+	ELinkedFilterMode FilterMode = ELinkedFilterMode::Linked;
+	
+	/** Ensures that the active filter is refiltered when a change is requested on it. The active filter is always registered with it. */
+	const TSharedRef<FFilterEvaluator> FilterEvaluator;
+	/** When the active filter bar changes, this is leveraged to orchestrate FFilterEvaluator::IncrementFilterBarUsage and DecrementFilterBarUsage. */
+	FScopedFilterBarUsage ActiveFilterBarUsage;
+	
+	/** Event when the filter mode changes, i.e. the instance returned by GetActiveFilterBar changes. */
+	FSimpleMulticastDelegate OnFilterModeChangedDelegate;
+	/** Event when the items displayed need to be refreshed, e.g. when ELinkedFilterMode changes, or the active filter bar refilters its items. */
+	FSimpleMulticastDelegate OnFilteredItemsChangedDelegate;
+	
+	void OnFilterBarFiltered(const TSharedRef<FSequencerFilterBar>& InFilterBar);
+};
+} // namespace UE::Sequencer
diff --git a/Engine/Source/Editor/Sequencer/Private/Filters/Linking/Mode/MakeLinkedFilterViewModelArgs.h b/Engine/Source/Editor/Sequencer/Private/Filters/Linking/Mode/MakeLinkedFilterViewModelArgs.h
new file mode 100644
--- /dev/null
+++ b/Engine/Source/Editor/Sequencer/Private/Filters/Linking/Mode/MakeLinkedFilterViewModelArgs.h
@@ -0,0 +1,36 @@
+﻿// Copyright Epic Games, Inc. All Rights Reserved.
+
+#pragma once
+
+#include "Containers/UnrealString.h"
+#include "Filters/FilterBarFlags.h"
+
+namespace UE::Sequencer
+{
+/** @see ILinkedFilterFactoryViewModel::MakeFilteringModel */
+struct FMakeLinkedFilterViewModelArgs
+{
+	/**
+	 * The sub-identifier to used for saving and loading filter entries.
+	 * 
+	 * The final identifier is the name of the settings (USequencerSettings::GetName) + "." + ConfigSubIdentifier.
+	 * Example: "LevelSequencerEditor.Linked", "LevelSequencerEditor.CurveEditor", etc.
+	 * 
+	 * @see USequencerSettings::TrackFilterBars.
+	 */
+	FString ConfigSubIdentifier;
+	
+	/** 
+	 * The filter flags to use for the unlinked filter bar.
+	 * 
+	 * By default, unlinked filter state should not be affected by the USequencerSettings::bIncludePinnedInFilter setting, so you likely don't want 
+	 * to set PinnedItemsCanSkipFiltering. 
+	 */
+	EFilterFlags UnlinkedFilterFlags = EFilterFlags::None;
+	
+	explicit FMakeLinkedFilterViewModelArgs(FString InConfigSubIdentifier, EFilterFlags InFilterFlags = EFilterFlags::None)
+		: ConfigSubIdentifier(MoveTemp(InConfigSubIdentifier))
+		, UnlinkedFilterFlags(InFilterFlags)
+		{}
+};
+} // namespace UE::Sequencer
diff --git a/Engine/Source/Editor/Sequencer/Private/Filters/Utils/LinkedFilteringViewUtils.cpp b/Engine/Source/Editor/Sequencer/Private/Filters/Utils/LinkedFilteringViewUtils.cpp
--- a/Engine/Source/Editor/Sequencer/Private/Filters/Utils/LinkedFilteringViewUtils.cpp
+++ b/Engine/Source/Editor/Sequencer/Private/Filters/Utils/LinkedFilteringViewUtils.cpp
@@ -2,7 +2,7 @@
 
 #include "LinkedFilteringViewUtils.h"
 
-#include "Filters/ViewModels/Linking/ILinkedFilterViewModel.h"
+#include "Filters/Linking/Mode/ILinkedFilterViewModel.h"
 #include "Misc/ConsoleVariables.h"
 #include "Templates/SharedPointer.h"
 
diff --git a/Engine/Source/Editor/Sequencer/Private/Filters/ViewModels/FilterEvaluator.cpp b/Engine/Source/Editor/Sequencer/Private/Filters/ViewModels/FilterEvaluator.cpp
deleted file mode 100644
diff --git a/Engine/Source/Editor/Sequencer/Private/Filters/ViewModels/FilterEvaluator.h b/Engine/Source/Editor/Sequencer/Private/Filters/ViewModels/FilterEvaluator.h
deleted file mode 100644
diff --git a/Engine/Source/Editor/Sequencer/Private/Filters/ViewModels/Linking/FilterAreaManager.cpp b/Engine/Source/Editor/Sequencer/Private/Filters/ViewModels/Linking/FilterAreaManager.cpp
deleted file mode 100644
diff --git a/Engine/Source/Editor/Sequencer/Private/Filters/ViewModels/Linking/FilterAreaManager.h b/Engine/Source/Editor/Sequencer/Private/Filters/ViewModels/Linking/FilterAreaManager.h
deleted file mode 100644
diff --git a/Engine/Source/Editor/Sequencer/Private/Filters/ViewModels/Linking/ILinkedFilterFactoryViewModel.h b/Engine/Source/Editor/Sequencer/Private/Filters/ViewModels/Linking/ILinkedFilterFactoryViewModel.h
deleted file mode 100644
diff --git a/Engine/Source/Editor/Sequencer/Private/Filters/ViewModels/Linking/ILinkedFilterViewModel.cpp b/Engine/Source/Editor/Sequencer/Private/Filters/ViewModels/Linking/ILinkedFilterViewModel.cpp
deleted file mode 100644
diff --git a/Engine/Source/Editor/Sequencer/Private/Filters/ViewModels/Linking/ILinkedFilterViewModel.h b/Engine/Source/Editor/Sequencer/Private/Filters/ViewModels/Linking/ILinkedFilterViewModel.h
deleted file mode 100644
diff --git a/Engine/Source/Editor/Sequencer/Private/Filters/ViewModels/Linking/LinkedFilterFactoryViewModel.cpp b/Engine/Source/Editor/Sequencer/Private/Filters/ViewModels/Linking/LinkedFilterFactoryViewModel.cpp
deleted file mode 100644
diff --git a/Engine/Source/Editor/Sequencer/Private/Filters/ViewModels/Linking/LinkedFilterFactoryViewModel.h b/Engine/Source/Editor/Sequencer/Private/Filters/ViewModels/Linking/LinkedFilterFactoryViewModel.h
deleted file mode 100644
diff --git a/Engine/Source/Editor/Sequencer/Private/Filters/ViewModels/Linking/LinkedFilterViewModel.cpp b/Engine/Source/Editor/Sequencer/Private/Filters/ViewModels/Linking/LinkedFilterViewModel.cpp
deleted file mode 100644
diff --git a/Engine/Source/Editor/Sequencer/Private/Filters/ViewModels/Linking/LinkedFilterViewModel.h b/Engine/Source/Editor/Sequencer/Private/Filters/ViewModels/Linking/LinkedFilterViewModel.h
deleted file mode 100644
diff --git a/Engine/Source/Editor/Sequencer/Private/Filters/ViewModels/Linking/MakeLinkedFilterViewModelArgs.h b/Engine/Source/Editor/Sequencer/Private/Filters/ViewModels/Linking/MakeLinkedFilterViewModelArgs.h
deleted file mode 100644
diff --git a/Engine/Source/Editor/Sequencer/Private/Filters/ViewModels/ScopedFilterBarUsage.h b/Engine/Source/Editor/Sequencer/Private/Filters/ViewModels/ScopedFilterBarUsage.h
deleted file mode 100644
diff --git a/Engine/Source/Editor/Sequencer/Private/Filters/Widgets/Linking/SFilterBarSwitcher.cpp b/Engine/Source/Editor/Sequencer/Private/Filters/Widgets/Linking/SFilterBarSwitcher.cpp
--- a/Engine/Source/Editor/Sequencer/Private/Filters/Widgets/Linking/SFilterBarSwitcher.cpp
+++ b/Engine/Source/Editor/Sequencer/Private/Filters/Widgets/Linking/SFilterBarSwitcher.cpp
@@ -5,7 +5,7 @@
 #include "Sequencer.h"
 #include "SFilterLinkStateSwitcher.h"
 #include "Filters/Widgets/SSequencerFilterBar.h"
-#include "Filters/ViewModels/Linking/LinkedFilterViewModel.h"
+#include "Filters/Linking/Mode/LinkedFilterViewModel.h"
 
 namespace UE::Sequencer
 {
diff --git a/Engine/Source/Editor/Sequencer/Private/Filters/Widgets/Linking/SFilterLinkStateSwitcher.cpp b/Engine/Source/Editor/Sequencer/Private/Filters/Widgets/Linking/SFilterLinkStateSwitcher.cpp
--- a/Engine/Source/Editor/Sequencer/Private/Filters/Widgets/Linking/SFilterLinkStateSwitcher.cpp
+++ b/Engine/Source/Editor/Sequencer/Private/Filters/Widgets/Linking/SFilterLinkStateSwitcher.cpp
@@ -2,7 +2,7 @@
 
 #include "SFilterLinkStateSwitcher.h"
 
-#include "Filters/ViewModels/Linking/ILinkedFilterViewModel.h"
+#include "Filters/Linking/Mode/ILinkedFilterViewModel.h"
 #include "Widgets/Layout/SWidgetSwitcher.h"
 
 namespace UE::Sequencer
diff --git a/Engine/Source/Editor/Sequencer/Private/Filters/Widgets/Linking/SLinkedFilteringToggleButton.cpp b/Engine/Source/Editor/Sequencer/Private/Filters/Widgets/Linking/SLinkedFilteringToggleButton.cpp
--- a/Engine/Source/Editor/Sequencer/Private/Filters/Widgets/Linking/SLinkedFilteringToggleButton.cpp
+++ b/Engine/Source/Editor/Sequencer/Private/Filters/Widgets/Linking/SLinkedFilteringToggleButton.cpp
@@ -2,7 +2,7 @@
 
 #include "SLinkedFilteringToggleButton.h"
 
-#include "Filters/ViewModels/Linking/ILinkedFilterViewModel.h"
+#include "Filters/Linking/Mode/ILinkedFilterViewModel.h"
 #include "Widgets/Images/SImage.h"
 #include "Widgets/Input/SCheckBox.h"
 
diff --git a/Engine/Source/Editor/Sequencer/Private/MVVM/CurveEditorExtension.cpp b/Engine/Source/Editor/Sequencer/Private/MVVM/CurveEditorExtension.cpp
--- a/Engine/Source/Editor/Sequencer/Private/MVVM/CurveEditorExtension.cpp
+++ b/Engine/Source/Editor/Sequencer/Private/MVVM/CurveEditorExtension.cpp
@@ -3,7 +3,7 @@
 #include "MVVM/CurveEditorExtension.h"
 
 #include "CurveEditor/SequencerCurveEditorApp.h"
-#include "Filters/ViewModels/Linking/ILinkedFilterViewModel.h"
+#include "Filters/Linking/Mode/ILinkedFilterViewModel.h"
 #include "FrameNumberDetailsCustomization.h"
 #include "Framework/Docking/TabManager.h"
 #include "MVVM/ViewModels/SequencerEditorViewModel.h"
diff --git a/Engine/Source/Editor/Sequencer/Private/MVVM/ViewModels/LinkedFilteringMediator.cpp b/Engine/Source/Editor/Sequencer/Private/MVVM/ViewModels/LinkedFilteringMediator.cpp
--- a/Engine/Source/Editor/Sequencer/Private/MVVM/ViewModels/LinkedFilteringMediator.cpp
+++ b/Engine/Source/Editor/Sequencer/Private/MVVM/ViewModels/LinkedFilteringMediator.cpp
@@ -5,8 +5,8 @@
 #include "Filters/SequencerFilterBar.h"
 #include "Filters/SequencerTrackFilterCommands.h"
 #include "Filters/Utils/LinkedFilteringViewUtils.h"
-#include "Filters/ViewModels/Linking/ILinkedFilterViewModel.h"
-#include "Filters/ViewModels/Linking/LinkedFilterViewModel.h"
+#include "Filters/Linking/Mode/ILinkedFilterViewModel.h"
+#include "Filters/Linking/Mode/LinkedFilterViewModel.h"
 #include "SequencerNodeTree.h"
 #include "Filters/FilterConfigIdentifiers.h"
 
diff --git a/Engine/Source/Editor/Sequencer/Private/MVVM/ViewModels/LinkedFilteringMediator.h b/Engine/Source/Editor/Sequencer/Private/MVVM/ViewModels/LinkedFilteringMediator.h
--- a/Engine/Source/Editor/Sequencer/Private/MVVM/ViewModels/LinkedFilteringMediator.h
+++ b/Engine/Source/Editor/Sequencer/Private/MVVM/ViewModels/LinkedFilteringMediator.h
@@ -2,7 +2,7 @@
 
 #pragma once
 
-#include "Filters/ViewModels/Linking/FilterAreaManager.h"
+#include "Filters/Linking/FilterAreaManager.h"
 #include "MVVM/Extensions/IOutlinerExtension.h"
 #include "MVVM/ViewModels/ViewModel.h"
 
diff --git a/Engine/Source/Editor/Sequencer/Private/MVVM/ViewModels/SequencerEditorViewModelImpl.cpp b/Engine/Source/Editor/Sequencer/Private/MVVM/ViewModels/SequencerEditorViewModelImpl.cpp
--- a/Engine/Source/Editor/Sequencer/Private/MVVM/ViewModels/SequencerEditorViewModelImpl.cpp
+++ b/Engine/Source/Editor/Sequencer/Private/MVVM/ViewModels/SequencerEditorViewModelImpl.cpp
@@ -3,9 +3,9 @@
 #include "SequencerEditorViewModelImpl.h"
 
 #include "Filters/FilterConfigIdentifiers.h"
-#include "Filters/ViewModels/Linking/LinkedFilterFactoryViewModel.h"
-#include "Filters/ViewModels/Linking/LinkedFilterViewModel.h"
-#include "Filters/ViewModels/Linking/MakeLinkedFilterViewModelArgs.h"
+#include "Filters/Linking/Mode/LinkedFilterFactoryViewModel.h"
+#include "Filters/Linking/Mode/LinkedFilterViewModel.h"
+#include "Filters/Linking/Mode/MakeLinkedFilterViewModelArgs.h"
 #include "LinkedFilteringMediator.h"
 #include "Sequencer.h"
 
diff --git a/Engine/Source/Editor/Sequencer/Private/Sequencer.cpp b/Engine/Source/Editor/Sequencer/Private/Sequencer.cpp
--- a/Engine/Source/Editor/Sequencer/Private/Sequencer.cpp
+++ b/Engine/Source/Editor/Sequencer/Private/Sequencer.cpp
@@ -192,7 +192,7 @@
 #include "CurveEditor/SequencerCurveEditorApp.h"
 #include "Decorations/MovieSceneMuteSoloDecoration.h"
 #include "Filters/SequencerTrackFilterCommands.h"
-#include "Filters/ViewModels/Linking/LinkedFilterViewModel.h"
+#include "Filters/Linking/Mode/LinkedFilterViewModel.h"
 #include "MVVM/ViewModels/LinkedFilteringMediator.h"
 #include "MVVM/ViewModels/SequencerEditorViewModelImpl.h"
 #include "TrackEditors/TimeWarpTrackEditor.h"
diff --git a/Engine/Source/Editor/Sequencer/Private/SequencerNodeTree.cpp b/Engine/Source/Editor/Sequencer/Private/SequencerNodeTree.cpp
--- a/Engine/Source/Editor/Sequencer/Private/SequencerNodeTree.cpp
+++ b/Engine/Source/Editor/Sequencer/Private/SequencerNodeTree.cpp
@@ -1,6 +1,8 @@
 // Copyright Epic Games, Inc. All Rights Reserved.
 
 #include "SequencerNodeTree.h"
+
+#include "Filters/Evaluation/FilterEvaluator.h"
 #include "Filters/SequencerFilterBar.h"
 #include "MVVM/ViewModels/ViewModelHierarchy.h"
 #include "MVVM/Selection/Selection.h"
@@ -41,7 +43,6 @@
 #include "SequencerLog.h"
 #include "SequencerCommonHelpers.h"
 #include "SSequencer.h"
-#include "Filters/ViewModels/FilterEvaluator.h"
 
 FSequencerNodeTree::~FSequencerNodeTree()
 {
diff --git a/Engine/Source/Editor/Sequencer/Private/Widgets/Sequencer/SSequencerSearchAndFilterRow.cpp b/Engine/Source/Editor/Sequencer/Private/Widgets/Sequencer/SSequencerSearchAndFilterRow.cpp
--- a/Engine/Source/Editor/Sequencer/Private/Widgets/Sequencer/SSequencerSearchAndFilterRow.cpp
+++ b/Engine/Source/Editor/Sequencer/Private/Widgets/Sequencer/SSequencerSearchAndFilterRow.cpp
@@ -7,8 +7,8 @@
 #include "Filters/Utils/HideIsolateViewModelUtils.h"
 #include "Filters/Utils/LinkedFilteringViewUtils.h"
 #include "Filters/ViewModels/HideIsolateShowViewModel.h"
-#include "Filters/ViewModels/Linking/ILinkedFilterViewModel.h"
-#include "Filters/ViewModels/Linking/LinkedFilterViewModel.h"
+#include "Filters/Linking/Mode/ILinkedFilterViewModel.h"
+#include "Filters/Linking/Mode/LinkedFilterViewModel.h"
 #include "Filters/Widgets/Linking/SFilterLinkStateSwitcher.h"
 #include "Filters/Widgets/Linking/SLinkedFilteringToggleButton.h"
 #include "Filters/Widgets/SSearchAndFilterWidget.h"
