Serquencer: Fix false positive static analysis warning of dangling reference

#Synced-CL 51626380

diff --git a/Engine/Source/Editor/Sequencer/Private/Filters/FilterBarInitArgs.h b/Engine/Source/Editor/Sequencer/Private/Filters/FilterBarInitArgs.h
--- a/Engine/Source/Editor/Sequencer/Private/Filters/FilterBarInitArgs.h
+++ b/Engine/Source/Editor/Sequencer/Private/Filters/FilterBarInitArgs.h
@@ -2,9 +2,10 @@
 
 #pragma once
 
-#include "FilterBarFlags.h"
 #include "Containers/UnrealString.h"
+#include "FilterBarFlags.h"
 #include "HAL/Platform.h"
+#include "Misc/NotNull.h"
 
 class FSequencer;
 
@@ -14,7 +15,7 @@ namespace UE::Sequencer
 struct FFilterBarInitArgs
 {
 	/** The Sequencer this filter bar is interacting with */
-	FSequencer& Sequencer;
+	TNotNull<FSequencer*> Sequencer;
 	
 	/**
 	 * Used for per filter bar instance config.
@@ -27,7 +28,7 @@ struct FFilterBarInitArgs
 	EFilterFlags Flags = EFilterFlags::None;
 
 	explicit FFilterBarInitArgs(FSequencer& InSequencer UE_LIFETIMEBOUND, const FString& InSubIdentifier, EFilterFlags InFlags = EFilterFlags::None)
-		: Sequencer(InSequencer)
+		: Sequencer(&InSequencer)
 		, SubIdentifier(InSubIdentifier)
 		, Flags(InFlags)
 	{}
diff --git a/Engine/Source/Editor/Sequencer/Private/Filters/SequencerFilterBar.cpp b/Engine/Source/Editor/Sequencer/Private/Filters/SequencerFilterBar.cpp
--- a/Engine/Source/Editor/Sequencer/Private/Filters/SequencerFilterBar.cpp
+++ b/Engine/Source/Editor/Sequencer/Private/Filters/SequencerFilterBar.cpp
@@ -41,7 +41,7 @@ using namespace UE::Sequencer;
 int32 FSequencerFilterBar::InstanceCount = 0;
 
 FSequencerFilterBar::FSequencerFilterBar(FFilterBarInitArgs InArgs)
-	: Sequencer(InArgs.Sequencer)
+	: Sequencer(*InArgs.Sequencer)
 	, CommandList(MakeShared<FUICommandList>())
 	, ClassTypeCategory(MakeShared<FFilterCategory>(LOCTEXT("ActorTypeFilterCategory", "Actor Type Filters"), FText()))
 	, ComponentTypeCategory(MakeShared<FFilterCategory>(LOCTEXT("ObjectTypeFilterCategory", "Object Type Filters"), FText()))
