Sequencer: Reset to new root sequence if the focused sequence is renamed

This fixes a crash when renaming a subsequence while viewing through it

#Synced-CL 51670918

diff --git a/Engine/Source/Editor/Sequencer/Private/Sequencer.cpp b/Engine/Source/Editor/Sequencer/Private/Sequencer.cpp
--- a/Engine/Source/Editor/Sequencer/Private/Sequencer.cpp
+++ b/Engine/Source/Editor/Sequencer/Private/Sequencer.cpp
@@ -197,6 +197,8 @@
 #include "MVVM/ViewModels/SequencerEditorViewModelImpl.h"
 #include "TrackEditors/TimeWarpTrackEditor.h"
 #include "Variants/MovieSceneTimeWarpGetter.h"
+#include "AssetRegistry/IAssetRegistry.h"
+#include "AssetRegistry/AssetRegistryModule.h"
 
 #define LOCTEXT_NAMESPACE "Sequencer"
 
@@ -557,6 +559,22 @@ void FSequencer::InitSequencer(const FSequencerInitParams& InitParams
 		AcquiredResources.Add([=] { FCoreUObjectDelegates::OnObjectsReplaced.Remove(OnObjectsReplacedHandle); });
 	}
 
+	if (FAssetRegistryModule* AssetRegistryModule = FModuleManager::GetModulePtr<FAssetRegistryModule>(TEXT("AssetRegistry")))
+	{
+		FDelegateHandle OnAssetRenamedHandle = AssetRegistryModule->Get().OnAssetRenamed().AddLambda([this](const FAssetData& AssetData, const FString& AssetOldName)
+		{
+			UMovieSceneSequence* MovieScene = GetFocusedMovieSceneSequence();
+			if (AssetData.GetAsset() == MovieScene)
+			{
+				if (UMovieSceneSequence* RootSequence = GetRootMovieSceneSequence())
+				{
+					ResetToNewRootSequence(*RootSequence);
+				}
+			}
+		});
+		AcquiredResources.Add([=] { AssetRegistryModule->Get().OnAssetRenamed().Remove(OnAssetRenamedHandle); });
+	}
+
 	ToolkitHost = InitParams.ToolkitHost;
 
 	PlaybackSpeed = 1.f;
