Sequencer: When right clicking on multiple bindings populate the binding array with all selected bindings of same class.

#Synced-CL 51508179

diff --git a/Engine/Source/Editor/Sequencer/Private/MVVM/ViewModels/ObjectBindingModel.cpp b/Engine/Source/Editor/Sequencer/Private/MVVM/ViewModels/ObjectBindingModel.cpp
--- a/Engine/Source/Editor/Sequencer/Private/MVVM/ViewModels/ObjectBindingModel.cpp
+++ b/Engine/Source/Editor/Sequencer/Private/MVVM/ViewModels/ObjectBindingModel.cpp
@@ -508,7 +508,7 @@ TSharedRef<SWidget> FObjectBindingModel::GetAddTrackMenuContent()
 	TSharedPtr<FSequencer> Sequencer = OwnerModel->GetSequencerImpl();
 	check(Sequencer);
 
-	UObject* BoundObject = GetFirstBoundObject();
+	UObject* BoundObject = GetFirstBoundObject(ObjectBindingID);
 	const UClass* MainSelectionObjectClass = FindObjectClass();
 
 	TArray<FGuid> ObjectBindings;
@@ -790,7 +790,7 @@ void FObjectBindingModel::HandleAddTrackSubMenuNew(FMenuBuilder& AddTrackMenuBui
 void FObjectBindingModel::HandlePropertyMenuItemExecute(FPropertyPath PropertyPath)
 {
 	TSharedPtr<FSequencer> Sequencer = OwnerModel->GetSequencerImpl();
-	UObject *BoundObject = GetFirstBoundObject();
+	UObject *BoundObject = GetFirstBoundObject(ObjectBindingID);
 	TArray<UObject*> KeyableBoundObjects;
 	if (BoundObject != nullptr)
 	{
@@ -835,7 +835,7 @@ void FObjectBindingModel::BuildContextMenu(FMenuBuilder& MenuBuilder)
 	TSharedPtr<FSequencerEditorViewModel> EditorViewModel = GetEditor();
 	FSequencer* Sequencer = EditorViewModel->GetSequencerImpl().Get();
 	ISequencerModule& SequencerModule = FModuleManager::GetModuleChecked<ISequencerModule>("Sequencer");
-	UObject* BoundObject = GetFirstBoundObject();
+	UObject* BoundObject = GetFirstBoundObject(ObjectBindingID);
 	const UClass* ObjectClass = FindObjectClass();
 	
 	TSharedPtr<FExtender> Extender = EditorViewModel->GetSequencerMenuExtender(
@@ -856,6 +856,24 @@ void FObjectBindingModel::BuildContextMenu(FMenuBuilder& MenuBuilder)
 	// Track editor extension.
 	TArray<FGuid> ObjectBindings;
 	ObjectBindings.Add(ObjectBindingID);
+
+	// Add other selected objects with the same class as the first one
+	for (TViewModelPtr<FObjectBindingModel> ObjectBindingNode : OwnerModel->GetEditor()->GetSelection()->Outliner.Filter<FObjectBindingModel>())
+	{
+		FGuid Guid = ObjectBindingNode->GetObjectGuid();
+		if (Guid != ObjectBindingID)
+		{
+			if (UObject* OtherBoundObject = GetFirstBoundObject(Guid))
+			{
+				if (ObjectClass == ObjectBindingNode->FindObjectClass())
+				{
+					ObjectBindings.Add(Guid);
+				}
+			}
+		}
+	}
+
+
 	for (const TSharedPtr<ISequencerTrackEditor>& TrackEditor : Sequencer->GetTrackEditors())
 	{
 		TrackEditor->BuildObjectBindingContextMenu(MenuBuilder, ObjectBindings, ObjectClass);
@@ -891,7 +909,7 @@ void FObjectBindingModel::BuildSidebarMenu(FMenuBuilder& MenuBuilder)
 	} 
 
 	ISequencerModule& SequencerModule = FModuleManager::GetModuleChecked<ISequencerModule>(TEXT("Sequencer"));
-	UObject* BoundObject = GetFirstBoundObject();
+	UObject* BoundObject = GetFirstBoundObject(ObjectBindingID);
 
 	const TSharedPtr<FExtender> Extender = EditorViewModel->GetSequencerMenuExtender(SequencerModule.GetSidebarExtensibilityManager()
 		, TArrayBuilder<UObject*>().Add(BoundObject), &FSequencerCustomizationInfo::OnBuildSidebarMenu, SharedThis(this));
@@ -1220,7 +1238,7 @@ void FObjectBindingModel::SetIsSoloed(bool bIsSoloed)
 	}
 }
 
-UObject* FObjectBindingModel::GetFirstBoundObject()
+UObject* FObjectBindingModel::GetFirstBoundObject(FGuid InGuid)
 {
 	TSharedPtr<ISequencer> Sequencer = OwnerModel->GetSequencer();
 	if (!Sequencer)
@@ -1228,7 +1246,7 @@ UObject* FObjectBindingModel::GetFirstBoundObject()
 		return nullptr;
 	}
 
-	for (TWeakObjectPtr<UObject> RuntimeObject : Sequencer->FindObjectsInCurrentSequence(ObjectBindingID))
+	for (TWeakObjectPtr<UObject> RuntimeObject : Sequencer->FindObjectsInCurrentSequence(InGuid))
 	{
 		UObject* BoundObject = RuntimeObject.Get();
 		if (BoundObject)
diff --git a/Engine/Source/Editor/Sequencer/Public/MVVM/ViewModels/ObjectBindingModel.h b/Engine/Source/Editor/Sequencer/Public/MVVM/ViewModels/ObjectBindingModel.h
--- a/Engine/Source/Editor/Sequencer/Public/MVVM/ViewModels/ObjectBindingModel.h
+++ b/Engine/Source/Editor/Sequencer/Public/MVVM/ViewModels/ObjectBindingModel.h
@@ -168,7 +168,7 @@ private:
 
 	UE_API void OnFinishedChangingDynamicBindingProperties(const FPropertyChangedEvent& ChangeEvent, TSharedPtr<FStructOnScope> ValueStruct);
 
-	UE_API UObject* GetFirstBoundObject();
+	UE_API UObject* GetFirstBoundObject(FGuid InGuid);
 
 protected:
 
