Sequencer: Fix temporarily stale tree items trigger failing ensure when adding & removing constraints.

The cause was that Control Rig calls ResetCurveEditor when adding or removing constraints.
Each call of RemoveTreeItem causes the UpdateCurveEditor to be called indirectly due to selection change delegates invoking it. That in turn would cause some tree items to have been removed already and the TreeItemIDToViewModelMap not containing any mapping as intermediate result.
The fix is to make FCurveModelSyncer::ResetCurveEditor uses FScopedCurveEditorTreeEventGuard to defer the intemediate opeations.

This bug seems to have been around for a long time but only recently became noticeable due to the ensure that was added.

#Synced-CL 51626380

diff --git a/Engine/Source/Editor/Sequencer/Private/CurveEditor/Integration/Selection/CurveModelSyncer.cpp b/Engine/Source/Editor/Sequencer/Private/CurveEditor/Integration/Selection/CurveModelSyncer.cpp
--- a/Engine/Source/Editor/Sequencer/Private/CurveEditor/Integration/Selection/CurveModelSyncer.cpp
+++ b/Engine/Source/Editor/Sequencer/Private/CurveEditor/Integration/Selection/CurveModelSyncer.cpp
@@ -105,6 +105,11 @@ void FCurveModelSyncer::RequestSyncSelection()
 
 void FCurveModelSyncer::ResetCurveEditor()
 {
+	// The RemoveTreeItem calls below would trigger follow-up events, such as selection change events.
+	// Each RemoveTreeItem would cause UpdateCurveEditor to be called: that's bad for performance but also causes incorrect behavior because 
+	// UpdateCurveEditor would be called with a half intact hierarchy. Just process everything once at the end.
+	const FScopedCurveEditorTreeEventGuard Guard(CurveEditor->GetTree());
+	
 	for (const TPair<TWeakPtr<FViewModel>, FCurveEditorTreeItemID>& Pair : ViewModelToTreeItemIDMap)
 	{
 		const FCurveEditorTreeItemID ItemID(Pair.Value);
