Sequencer: Use TimeBounds when playing back in editor.

Time bounds will return the selection range if it exists, otherwise RootTimeBounds. This matches the behavior of rewind, jumping from key to key, etc.

#Synced-CL 51770305

diff --git a/Engine/Source/Editor/Sequencer/Private/Sequencer.cpp b/Engine/Source/Editor/Sequencer/Private/Sequencer.cpp
--- a/Engine/Source/Editor/Sequencer/Private/Sequencer.cpp
+++ b/Engine/Source/Editor/Sequencer/Private/Sequencer.cpp
@@ -5234,16 +5234,28 @@ FReply FSequencer::OnPlay(bool bTogglePlay)
 	}
 	else
 	{
-		TRange<FFrameNumber> TimeBounds = GetTimeBounds();
+		TRange<FFrameNumber> LocalTimeBounds = GetTimeBounds();
+
+		FMovieSceneInverseSequenceTransform LocalToRootTransform = RootToUnwarpedLocalTransform.Inverse();
+
+		TOptional<FFrameTime> MinInclusiveTime = LocalToRootTransform.TryTransformTime(UE::MovieScene::DiscreteInclusiveLower(LocalTimeBounds));
+		TOptional<FFrameTime> MaxInclusiveTime = LocalToRootTransform.TryTransformTime(UE::MovieScene::DiscreteExclusiveUpper(LocalTimeBounds) - 1);
 
-		FFrameNumber MinInclusiveTime = UE::MovieScene::DiscreteInclusiveLower(TimeBounds);
-		FFrameNumber MaxInclusiveTime = UE::MovieScene::DiscreteExclusiveUpper(TimeBounds) - 1;
-		
 		FFrameTime GlobalTime = GlobalPlaybackWarpTransform.TransformTime(GetGlobalTime().Time);
 
-		if (GlobalTime <= MinInclusiveTime || GlobalTime >= MaxInclusiveTime)
+		if (PlaybackSpeed > 0)
 		{
-			SetGlobalTime(PlaybackSpeed > 0 ? MinInclusiveTime : MaxInclusiveTime);
+			if (MinInclusiveTime.IsSet() && GlobalTime <= MinInclusiveTime.GetValue())
+			{
+				SetGlobalTime(MinInclusiveTime.GetValue());
+			}
+		}
+		else
+		{
+			if (MaxInclusiveTime.IsSet() && GlobalTime >= MaxInclusiveTime.GetValue())
+			{
+				SetGlobalTime(MaxInclusiveTime.GetValue());
+			}
 		}
 
 		SpeedIndexBeforePlay = CurrentSpeedIndex;
