Sequencer: The selected filter has been adjusted to show all child properties of a control when a child property is selected.

Example: You select "Dragon.ChestCtrl.Location.X"
- Previous: only shows Location.X
- New: show Location.Y, Location.Z, Rotation, Scale, etc.

This change is relevant for when the selection is made in the Curve Editor.
This change does NOT directly affect when the selection is made in Sequencer's Outliner because selection changes made in the Outliner do not cause filter reruns; selection changes made by [REDACTED_PERSON] cause filter reevaluation. Another way to run into this path is to search for something in Sequencer's Outliner and then clear the search.

#Synced-CL 51760482

diff --git a/Engine/Plugins/Animation/ControlRig/Source/ControlRigEditor/Private/Sequencer/ControlRigSelectionFilterExtension.cpp b/Engine/Plugins/Animation/ControlRig/Source/ControlRigEditor/Private/Sequencer/ControlRigSelectionFilterExtension.cpp
new file mode 100644
--- /dev/null
+++ b/Engine/Plugins/Animation/ControlRig/Source/ControlRigEditor/Private/Sequencer/ControlRigSelectionFilterExtension.cpp
@@ -0,0 +1,105 @@
+﻿// Copyright Epic Games, Inc. All Rights Reserved.
+
+#include "ControlRigSelectionFilterExtension.h"
+
+#include "Filters/Filters/SequencerTrackFilter_Selected.h"
+#include "MVVM/Selection/Selection.h"
+#include "MVVM/ViewModels/CategoryModel.h"
+#include "MVVM/ViewModels/SequencerEditorViewModel.h"
+#include "MVVM/ViewModels/TrackModel.h"
+#include "MVVM/ViewModels/ViewModelIterators.h"
+#include "Sequencer/MovieSceneControlRigParameterTrack.h"
+
+namespace UE::ControlRig::SelectionFilterDetail
+{
+using EBreakBehavior = UControlRigSelectionFilterExtension::EBreakBehavior;
+
+/** 
+ * If InSelected is under any FViewModel representing a control, invoke InCallback with all of children of that control.
+ * 
+ * For example, if "Dragon.ChestCtrl.Location.X" is selected, then invoke with "Dragon.ChestCtrl.Location.Y", "Dragon.ChestCtrl.Rotation.X", etc.
+ * Invokes nothing if "Dragon.ChestCtrl" or above is selected.
+ */
+static void EnumerateSiblingsIfUnderControlRigTrack(
+	const Sequencer::FViewModelPtr& InSelected, 
+	TFunctionRef<EBreakBehavior(const Sequencer::FViewModelPtr&)> InCallback,
+	TArray<Sequencer::FViewModelPtr>& InAlreadyVisited
+	)
+{
+	using namespace UE::Sequencer;
+	// Skip if "Dragon.ChestCtrl" was selected. If "Dragon.Chest.Control.Foo" was selected, proceed.
+	const bool bIsUnderControl = CastViewModel<FChannelGroupModel>(InSelected) != nullptr
+		|| CastViewModel<FChannelModel>(InSelected) != nullptr;
+	if (!bIsUnderControl)
+	{
+		return;
+	}
+	
+	for (const FViewModelPtr& Current : FParentModelIterator(InSelected, true))
+	{
+		// This is an optimization to avoid reiterating the subhierarchy in case the user has selected e.g. "Location.X" and "Location.Y".
+		if (InAlreadyVisited.Contains(Current))
+		{
+			return;
+		}
+		InAlreadyVisited.Add(Current);
+			
+		const FViewModelPtr Parent = Current->GetParent();
+		const TViewModelPtr<FTrackModel> TrackModel = CastViewModel<FTrackModel>(Parent);
+		if (!TrackModel)
+		{
+			continue;
+		}
+			
+		UMovieSceneControlRigParameterTrack* ControlRigTrack = Cast<UMovieSceneControlRigParameterTrack>(TrackModel->GetTrack());
+		if (!ControlRigTrack)
+		{
+			continue;
+		}
+			
+		for (const FViewModelPtr& Child : FParentFirstChildIterator(Current))
+		{
+			if (InCallback(Child) == EBreakBehavior::Break)
+			{
+				return;
+			}
+		}
+	}
+}
+
+static void EnumerateParentControlChildren(
+	const ISequencer& InSequencer,
+	TFunctionRef<EBreakBehavior(const Sequencer::FViewModelPtr&)> InCallback
+	)
+{
+	using namespace UE::Sequencer;
+	const TSharedPtr<FSequencerEditorViewModel> ViewModel = InSequencer.GetViewModel();
+	if (!ViewModel.IsValid())
+	{
+		return;
+	}
+
+	const TSharedPtr<FSequencerSelection> SequencerSelection = ViewModel->GetSelection();
+	if (!SequencerSelection.IsValid())
+	{
+		return;
+	}
+	
+	TArray<FViewModelPtr> AllReadyVisited;
+	for (const FViewModelPtr& Selected : SequencerSelection->Outliner)
+	{
+		EnumerateSiblingsIfUnderControlRigTrack(Selected, InCallback, AllReadyVisited);
+	}
+}
+}
+
+void UControlRigSelectionFilterExtension::EnumerateViewModelsConsideredAsSelected(
+	const FSequencerTrackFilter_Selected& InFilter,
+	TFunctionRef<EBreakBehavior(const UE::Sequencer::FViewModelPtr&)> InCallback
+	)
+{
+	using namespace UE::ControlRig::SelectionFilterDetail;
+	
+	// E.g. if "Dragon.ChestCtrl.Location.X" is selected, then also show the other child properties: Location, Rotation, etc.
+	EnumerateParentControlChildren(InFilter.GetSequencer(), InCallback);
+}
diff --git a/Engine/Plugins/Animation/ControlRig/Source/ControlRigEditor/Private/Sequencer/ControlRigSelectionFilterExtension.h b/Engine/Plugins/Animation/ControlRig/Source/ControlRigEditor/Private/Sequencer/ControlRigSelectionFilterExtension.h
new file mode 100644
--- /dev/null
+++ b/Engine/Plugins/Animation/ControlRig/Source/ControlRigEditor/Private/Sequencer/ControlRigSelectionFilterExtension.h
@@ -0,0 +1,21 @@
+﻿// Copyright Epic Games, Inc. All Rights Reserved.
+
+#pragma once
+
+#include "Filters/SequencerTrackFilterSelectedExtension.h"
+#include "ControlRigSelectionFilterExtension.generated.h"
+
+/**
+ * Extends FSequencerTrackFilter_Selected such that selecting e.g. Dragon.ChestCtrl.Location.X shows Location.Y, Rotation, Scale, etc., as well.
+ */
+UCLASS()
+class UControlRigSelectionFilterExtension : public USequencerTrackFilterSelectedExtension
+{
+	GENERATED_BODY()
+public:
+	
+	virtual void EnumerateViewModelsConsideredAsSelected(
+		const FSequencerTrackFilter_Selected& InFilter,
+		TFunctionRef<EBreakBehavior(const UE::Sequencer::FViewModelPtr&)> InCallback
+		) override;
+};
diff --git a/Engine/Source/Editor/Sequencer/Private/Filters/Filters/SequencerTrackFilter_Selected.cpp b/Engine/Source/Editor/Sequencer/Private/Filters/Filters/SequencerTrackFilter_Selected.cpp
--- a/Engine/Source/Editor/Sequencer/Private/Filters/Filters/SequencerTrackFilter_Selected.cpp
+++ b/Engine/Source/Editor/Sequencer/Private/Filters/Filters/SequencerTrackFilter_Selected.cpp
@@ -1,22 +1,92 @@
 ﻿// Copyright Epic Games, Inc. All Rights Reserved.
 
 #include "Filters/Filters/SequencerTrackFilter_Selected.h"
+
 #include "Filters/SequencerFilterBar.h"
 #include "Filters/SequencerTrackFilterCommands.h"
+#include "Filters/SequencerTrackFilterSelectedExtension.h"
 #include "ISequencer.h"
 #include "MVVM/Selection/Selection.h"
 #include "MVVM/ViewModels/SequencerEditorViewModel.h"
 #include "Selection.h"
+#include "UObject/UObjectIterator.h"
 
 using namespace UE::Sequencer;
 
 #define LOCTEXT_NAMESPACE "SequencerTrackFilter_Selected"
 
-FSequencerTrackFilter_Selected::FSequencerTrackFilter_Selected(ISequencerTrackFilters& InFilterInterface, TSharedPtr<FFilterCategory> InCategory)
-	: FSequencerTrackFilter(InFilterInterface, MoveTemp(InCategory))
+
+
+namespace UE::Sequencer::SelectionDetail
+{
+static TArray<TWeakObjectPtr<USequencerTrackFilterSelectedExtension>> GetSelectionExtensions()
 {
+	TArray<TWeakObjectPtr<USequencerTrackFilterSelectedExtension>> Extensions;
+	for (TObjectIterator<USequencerTrackFilterSelectedExtension> ExtensionIt(RF_NoFlags); ExtensionIt; ++ExtensionIt)
+	{
+		USequencerTrackFilterSelectedExtension* PotentialExtension = *ExtensionIt;
+		if (PotentialExtension
+			&& PotentialExtension->HasAnyFlags(RF_ClassDefaultObject)
+			&& !PotentialExtension->GetClass()->HasAnyClassFlags(CLASS_Deprecated | CLASS_Abstract))
+		{
+			Extensions.Emplace(PotentialExtension);
+		}
+	}
+	return Extensions;
 }
 
+/** @return Whether InItem is a child of InSelectedNode. */
+static bool IsPartOfHierarchy(const FViewModelPtr& InItem, const FViewModelPtr& InSelectedNode)
+{
+	if (!InSelectedNode.IsValid())
+	{
+		return false;
+	}
+
+	for (FParentFirstChildIterator It(InSelectedNode, /*bIncludeThis=*/true); It; ++It)
+	{
+		if (*It == InItem)
+		{
+			return true;
+		}
+	}
+	
+	return false;
+}
+
+/** @return Whether InItem is a child of the FViewModels enumerated by USequencerTrackFilterSelectedExtension. */
+static bool IsPartOfExtendedSelection(
+	const FViewModelPtr& InItem, 
+	const FSequencerTrackFilter_Selected& InFilter, 
+	TConstArrayView<TWeakObjectPtr<USequencerTrackFilterSelectedExtension>> InWeakExtensions
+	)
+{
+	using EBreakBehavior = USequencerTrackFilterSelectedExtension::EBreakBehavior;
+	const auto ProcessExtension = [&InItem, &InFilter](const TWeakObjectPtr<USequencerTrackFilterSelectedExtension>& WeakExtension)
+	{
+		if (!WeakExtension.IsValid())
+		{
+			return false;
+		}
+		
+		bool bIsPartOfSelection = false;
+		WeakExtension->EnumerateViewModelsConsideredAsSelected(InFilter, [&InItem, &bIsPartOfSelection](const FViewModelPtr& ConsiderSelected)
+		{
+			bIsPartOfSelection = IsPartOfHierarchy(InItem, ConsiderSelected);
+			return bIsPartOfSelection ? EBreakBehavior::Break : EBreakBehavior::Continue;
+		});
+		return bIsPartOfSelection;
+	};
+	
+	return Algo::AnyOf(InWeakExtensions, ProcessExtension);
+}
+}
+
+FSequencerTrackFilter_Selected::FSequencerTrackFilter_Selected(ISequencerTrackFilters& InFilterInterface, TSharedPtr<FFilterCategory> InCategory)
+	: FSequencerTrackFilter(InFilterInterface, MoveTemp(InCategory))
+	, CachedWeakExtensions(SelectionDetail::GetSelectionExtensions())
+{}
+
 FSequencerTrackFilter_Selected::~FSequencerTrackFilter_Selected()
 {
 	UnbindSelectionChanged();
@@ -97,23 +167,13 @@ bool FSequencerTrackFilter_Selected::PassesFilter(FSequencerTrackFilterType InIt
 		return false;
 	}
 
-	for (const FViewModelPtr& SelectedNode : SequencerSelection->Outliner)
+	const bool bIsPartOfSequencerSelection = Algo::AnyOf(SequencerSelection->Outliner, [InItem](const FViewModelPtr& Selected)
 	{
-		if (!SelectedNode.IsValid())
-		{
-			continue;
-		}
-
-		for (FParentFirstChildIterator It(SelectedNode, /*bIncludeThis=*/true); It; ++It)
-		{
-			if (*It == InItem)
-			{
-				return true;
-			}
-		}
-	}
-
-	return false;
+		return SelectionDetail::IsPartOfHierarchy(InItem, Selected);
+	});
+	return bIsPartOfSequencerSelection 
+		// Allow external logic to make us pretend like additional items are selected without actually modifying FSequencerSelection::Outliner.
+		|| SelectionDetail::IsPartOfExtendedSelection(InItem, *this, CachedWeakExtensions);
 }
 
 void FSequencerTrackFilter_Selected::OnSelectionChanged(UObject* const InObject)
diff --git a/Engine/Source/Editor/Sequencer/Public/Filters/Filters/SequencerTrackFilter_Selected.h b/Engine/Source/Editor/Sequencer/Public/Filters/Filters/SequencerTrackFilter_Selected.h
--- a/Engine/Source/Editor/Sequencer/Public/Filters/Filters/SequencerTrackFilter_Selected.h
+++ b/Engine/Source/Editor/Sequencer/Public/Filters/Filters/SequencerTrackFilter_Selected.h
@@ -4,6 +4,8 @@
 
 #include "Filters/SequencerTrackFilterBase.h"
 
+class USequencerTrackFilterSelectedExtension;
+
 class FSequencerTrackFilter_Selected : public FSequencerTrackFilter
 {
 public:
@@ -38,4 +40,9 @@ protected:
 	void OnSelectionChanged(UObject* const InObject);
 
 	FDelegateHandle OnSelectionChangedHandle;
+	
+private:
+	
+	/** Objects extending this filter. Cached at construction to speed up filtering. */
+	const TArray<TWeakObjectPtr<USequencerTrackFilterSelectedExtension>> CachedWeakExtensions;
 };
diff --git a/Engine/Source/Editor/Sequencer/Public/Filters/SequencerTrackFilterSelectedExtension.h b/Engine/Source/Editor/Sequencer/Public/Filters/SequencerTrackFilterSelectedExtension.h
new file mode 100644
--- /dev/null
+++ b/Engine/Source/Editor/Sequencer/Public/Filters/SequencerTrackFilterSelectedExtension.h
@@ -0,0 +1,26 @@
+﻿// Copyright Epic Games, Inc. All Rights Reserved.
+
+#pragma once
+
+#include "SequencerCoreFwd.h"
+#include "Templates/FunctionFwd.h"
+#include "SequencerTrackFilterSelectedExtension.generated.h"
+
+class FSequencerTrackFilter_Selected;
+class ISequencerTrackFilters;
+
+/** Derive from this class to customize FSequencerTrackFilter_Selected. */
+UCLASS(MinimalAPI, Abstract)
+class USequencerTrackFilterSelectedExtension : public UObject
+{
+	GENERATED_BODY()
+public:
+	
+	enum class EBreakBehavior : uint8 { Break, Continue };
+	
+	/** Enumerates all view models that the FSequencerTrackFilter_Selected should treat as selected on top of what is selected in Sequencer. */
+	virtual void EnumerateViewModelsConsideredAsSelected(
+		const FSequencerTrackFilter_Selected& InFilter, TFunctionRef<EBreakBehavior(const UE::Sequencer::FViewModelPtr&)> InCallback
+		)
+	{}
+};
