Sequencer: Fix Curve Editor auto-expanding everything when making selection in Sequencer Outliner.

#Synced-CL 51760482

diff --git a/Engine/Source/Editor/Sequencer/Private/CurveEditor/Views/SelectionSyncUtils.cpp b/Engine/Source/Editor/Sequencer/Private/CurveEditor/Views/SelectionSyncUtils.cpp
--- a/Engine/Source/Editor/Sequencer/Private/CurveEditor/Views/SelectionSyncUtils.cpp
+++ b/Engine/Source/Editor/Sequencer/Private/CurveEditor/Views/SelectionSyncUtils.cpp
@@ -128,11 +128,21 @@ void ApplySequencerToCurveEditorSelection(
 
 namespace SyncDetail
 {
+/** Expands all parents such that InItem is visible. Does not expand InItem itself. */
 static void SetParentsExpanded(const FCurveEditorTree& InTree, const FCurveEditorTreeItemID& InItem, SCurveEditorTree& InTreeView, bool bInExpanded)
 {
-	for (const FCurveEditorTreeItem* Item = InTree.FindItem(InItem); Item; Item = InTree.FindItem(Item->GetParentID()))
+	const auto GetParent = [&InTree](const FCurveEditorTreeItem* Item){ return InTree.FindItem(Item->GetParentID()); };
+	
+	const FCurveEditorTreeItem* TreeItem = InTree.FindItem(InItem);
+	const FCurveEditorTreeItem* StartItem = TreeItem ? GetParent(TreeItem) : nullptr;
+	if (!StartItem)
+	{
+		return;
+	}
+	
+	for (const FCurveEditorTreeItem* Current = StartItem; Current; Current = GetParent(Current))
 	{
-		InTreeView.SetItemExpansion(Item->GetID(), bInExpanded);
+		InTreeView.SetItemExpansion(Current->GetID(), bInExpanded);
 	}
 }
 }
@@ -141,7 +151,7 @@ void ExpandSelectedCurveEditorItems(FCurveEditorTree& InTree, SCurveEditorTree&
 {
 	for (const TPair<FCurveEditorTreeItemID, ECurveEditorTreeSelectionState>& Pair : InTree.GetSelection())
 	{
-		if (Pair.Value != ECurveEditorTreeSelectionState::None)
+		if (Pair.Value == ECurveEditorTreeSelectionState::Explicit)
 		{
 			SyncDetail::SetParentsExpanded(InTree, Pair.Key, InTreeView, true);
 		}
