Sequencer: Remove unused filter state subscription in SSequencer.

#Synced-CL 51760482

diff --git a/Engine/Source/Editor/Sequencer/Private/SSequencer.cpp b/Engine/Source/Editor/Sequencer/Private/SSequencer.cpp
--- a/Engine/Source/Editor/Sequencer/Private/SSequencer.cpp
+++ b/Engine/Source/Editor/Sequencer/Private/SSequencer.cpp
@@ -362,12 +362,6 @@ void SSequencer::Construct(const FArguments& InArgs, TSharedRef<FSequencer> InSe
 		SequencerWidgets.CreateTimeRangeSlider(TimeSliderControllerRef)
 	);
 
-	// TODO UE-363287
-	if (const TSharedPtr<FSequencerFilterBar> FilterBar = GetFilterBar())
-	{
-		FilterBar->OnFiltersChanged().AddSP(this, &SSequencer::OnTrackFiltersChanged);
-	}
-
 	OutlinerAndTrackAreaContent = ConstructOutlinerAndTrackArea(InArgs, TimeSliderArgs);
 
 	FilterRow = SNew(SSequencerSearchAndFilterRow)
@@ -4046,16 +4040,6 @@ TSharedPtr<FSequencerFilterBar> SSequencer::GetFilterBar() const
 	return nullptr;
 }
 
-void SSequencer::OnTrackFiltersChanged(const ESequencerFilterChange InChangeType, const TSharedRef<FSequencerTrackFilter>& InFilter)
-{
-	RebuildFilterBarContent();
-
-	if (SequencerTreeFilterStatusBar.IsValid())
-	{
-		SequencerTreeFilterStatusBar->UpdateText();
-	}
-}
-
 void SSequencer::RebuildSearchAndFilterRow()
 {
 	using namespace UE::Sequencer;
diff --git a/Engine/Source/Editor/Sequencer/Private/SSequencer.h b/Engine/Source/Editor/Sequencer/Private/SSequencer.h
--- a/Engine/Source/Editor/Sequencer/Private/SSequencer.h
+++ b/Engine/Source/Editor/Sequencer/Private/SSequencer.h
@@ -702,8 +702,6 @@ private:
 
 	void OnSidebarStateChanged(const FSidebarState& InNewState);
 
-	void OnTrackFiltersChanged(const ESequencerFilterChange InChangeType, const TSharedRef<FSequencerTrackFilter>& InFilter);
-
 	TSharedPtr<FTabManager> GetTabManager() const;
 
 	void RegisterSelectionDetailsTab();
