Sequencer: Add "Auto Expand Curve Editor on Selection" to Curve Editor view options. This setting causes items to be expanded in Curve Editor when clicked in Sequencer's Outliner. It is similar to "Auto Expand Nodes on Selection".

#Synced-CL 51881412

diff --git a/Engine/Source/Editor/Sequencer/Private/CurveEditor/Integration/CurveEditorCommandBinder.cpp b/Engine/Source/Editor/Sequencer/Private/CurveEditor/Integration/CurveEditorCommandBinder.cpp
--- a/Engine/Source/Editor/Sequencer/Private/CurveEditor/Integration/CurveEditorCommandBinder.cpp
+++ b/Engine/Source/Editor/Sequencer/Private/CurveEditor/Integration/CurveEditorCommandBinder.cpp
@@ -70,6 +70,8 @@ void FCurveCommandBinder::InitSequencerCommandBindings()
 	CurveEditorSharedBindings->MapAction(Commands.AddRotationKey, *SequencerCommandBindings->GetActionForCommand(Commands.AddRotationKey));
 	CurveEditorSharedBindings->MapAction(Commands.AddScaleKey, *SequencerCommandBindings->GetActionForCommand(Commands.AddScaleKey));
 	
+	CurveEditorSharedBindings->MapAction(Commands.ToggleAutoExpandCurveEditorOnSelection, *SequencerCommandBindings->GetActionForCommand(Commands.ToggleAutoExpandCurveEditorOnSelection));
+	
 	FilterArea.GetFilterAreaCommandList()->Append(CurveEditorSharedBindings.ToSharedRef());
 	CurveEditor->GetCommands()->Append(CurveEditorSharedBindings.ToSharedRef());
 	
diff --git a/Engine/Source/Editor/Sequencer/Private/CurveEditor/Widgets/Tree/SCurveTreeContent.cpp b/Engine/Source/Editor/Sequencer/Private/CurveEditor/Widgets/Tree/SCurveTreeContent.cpp
--- a/Engine/Source/Editor/Sequencer/Private/CurveEditor/Widgets/Tree/SCurveTreeContent.cpp
+++ b/Engine/Source/Editor/Sequencer/Private/CurveEditor/Widgets/Tree/SCurveTreeContent.cpp
@@ -64,11 +64,18 @@ void SCurveTreeContent::FocusPlayTimeDisplay()
 
 void SCurveTreeContent::SyncSelection() const
 {
-	if (const TSharedPtr<FSequencer> Sequencer = WeakSequencer.Pin())
+	const TSharedPtr<FSequencer> Sequencer = WeakSequencer.Pin();
+	if (!Sequencer)
+	{
+		return;
+	}
+	
+	ApplySequencerToCurveEditorSelection(
+		Sequencer->GetViewModel()->GetSelection(), CurveEditor.ToSharedRef(), CurveEditorTreeView.ToSharedRef()
+		);
+	
+	if (Sequencer->GetSequencerSettings()->GetAutoExpandCurveEditorOnSelection())
 	{
-		Sequencer::ApplySequencerToCurveEditorSelection(
-			Sequencer->GetViewModel()->GetSelection(), CurveEditor.ToSharedRef(), CurveEditorTreeView.ToSharedRef()
-			);
 		ExpandSelectedCurveEditorItems(*CurveEditor->GetTree(), *CurveEditorTreeView.ToSharedRef());
 	}
 }
diff --git a/Engine/Source/Editor/Sequencer/Private/Filters/Menus/SequencerViewOptionsMenu.cpp b/Engine/Source/Editor/Sequencer/Private/Filters/Menus/SequencerViewOptionsMenu.cpp
--- a/Engine/Source/Editor/Sequencer/Private/Filters/Menus/SequencerViewOptionsMenu.cpp
+++ b/Engine/Source/Editor/Sequencer/Private/Filters/Menus/SequencerViewOptionsMenu.cpp
@@ -93,7 +93,7 @@ void FSequencerViewOptionsMenu::PopulateForCurveEditor(UToolMenu* const InMenu,
 		WeakSequencer = Sequencer;
     	
 		PopulateFiltersSection(*InMenu);
-		// TODO UE-365214: PopulateSortAndOrganizeSection(*InMenu);
+		PopulateCurveEditorOrganizeSection(*InMenu);
 		ViewOptionsDetail::AddFilterVisibilitySection(*InMenu);
 	}
 }
@@ -208,6 +208,14 @@ void FSequencerViewOptionsMenu::PopulateFilterOptionsSection(UToolMenu& InMenu)
 	ViewOptionsDetail::AddFilterVisibilitySection(InMenu);
 }
 
+void FSequencerViewOptionsMenu::PopulateCurveEditorOrganizeSection(UToolMenu& InMenu)
+{
+	FToolMenuSection& Section = InMenu.FindOrAddSection(TEXT("CurveEditor.OrganizeAndSort"), LOCTEXT("CurveEditor.OrganizeAndSortHeader", "Organize and Sort"));
+	
+	Section.AddMenuEntry(FSequencerCommands::Get().ToggleAutoExpandCurveEditorOnSelection);
+	// TODO UE-365214: add sorting and expansion capabilities
+}
+
 bool FSequencerViewOptionsMenu::IsIncludePinnedInFilter() const
 {
 	const TSharedPtr<FSequencer> Sequencer = WeakSequencer.Pin();
diff --git a/Engine/Source/Editor/Sequencer/Private/Filters/Menus/SequencerViewOptionsMenu.h b/Engine/Source/Editor/Sequencer/Private/Filters/Menus/SequencerViewOptionsMenu.h
--- a/Engine/Source/Editor/Sequencer/Private/Filters/Menus/SequencerViewOptionsMenu.h
+++ b/Engine/Source/Editor/Sequencer/Private/Filters/Menus/SequencerViewOptionsMenu.h
@@ -32,6 +32,8 @@ protected:
 	void PopulateFiltersSection(UToolMenu& InMenu);
 	void PopulateSortAndOrganizeSection(UToolMenu& InMenu);
 	void PopulateFilterOptionsSection(UToolMenu& InMenu);
+	
+	void PopulateCurveEditorOrganizeSection(UToolMenu& InMenu);
 
 	bool IsIncludePinnedInFilter() const;
 	void ToggleIncludePinnedInFilter();
diff --git a/Engine/Source/Editor/Sequencer/Private/Sequencer.cpp b/Engine/Source/Editor/Sequencer/Private/Sequencer.cpp
--- a/Engine/Source/Editor/Sequencer/Private/Sequencer.cpp
+++ b/Engine/Source/Editor/Sequencer/Private/Sequencer.cpp
@@ -11432,6 +11432,12 @@ void FSequencer::BindCommands()
 		FExecuteAction::CreateLambda([this] { Settings->SetAutoExpandNodesOnSelection(!Settings->GetAutoExpandNodesOnSelection()); }),
 		FCanExecuteAction::CreateLambda([] { return true; }),
 		FIsActionChecked::CreateLambda([this] { return Settings->GetAutoExpandNodesOnSelection(); }) );
+	
+	SequencerCommandBindings->MapAction(
+		Commands.ToggleAutoExpandCurveEditorOnSelection,
+		FExecuteAction::CreateLambda([this] { Settings->SetAutoExpandCurveEditorOnSelection(!Settings->GetAutoExpandCurveEditorOnSelection()); }),
+		FCanExecuteAction::CreateLambda([] { return true; }),
+		FIsActionChecked::CreateLambda([this] { return Settings->GetAutoExpandCurveEditorOnSelection(); }) );
 
 	SequencerCommandBindings->MapAction(
 		Commands.ToggleRestoreOriginalViewportOnCameraCutUnlock,
diff --git a/Engine/Source/Editor/Sequencer/Private/SequencerCommands.cpp b/Engine/Source/Editor/Sequencer/Private/SequencerCommands.cpp
--- a/Engine/Source/Editor/Sequencer/Private/SequencerCommands.cpp
+++ b/Engine/Source/Editor/Sequencer/Private/SequencerCommands.cpp
@@ -64,6 +64,8 @@ void FSequencerCommands::RegisterCommands()
 	UI_COMMAND( SetPlayheadScrubbingLimitSelectionRange, "Selection Range", "While scrubbing, restrict the playhead to the selection range", EUserInterfaceActionType::RadioButton, FInputChord() );
 
 	UI_COMMAND( ToggleAutoExpandNodesOnSelection, "Auto Expand Nodes on Selection", "Toggle auto expanding the outliner tree on child selection", EUserInterfaceActionType::ToggleButton, FInputChord());
+	UI_COMMAND( ToggleAutoExpandCurveEditorOnSelection, "Auto Expand Curve Editor on Selection", "Toggles whether selecting items in Sequencer's Outliner will expand Curve Editor's tree view to show the item.", EUserInterfaceActionType::ToggleButton, FInputChord());
+	
 	UI_COMMAND( ToggleRestoreOriginalViewportOnCameraCutUnlock, "Restore Cinematic Viewports", "Toggle whether unlocking a camera cut track should return the viewport to its original location, or keep it where the camera cut was", EUserInterfaceActionType::ToggleButton, FInputChord() );
 	UI_COMMAND( TogglePreviewCameraCutsInSimulate, "Preview Camera Cuts in Simulate", "Toggle whether camera cuts are evaluated in SIE, or after ejecting from the player controller in PIE", EUserInterfaceActionType::ToggleButton, FInputChord() );
 	UI_COMMAND( ToggleExpandCollapseNodes, "Expand/Collapse Nodes", "Toggle expand or collapse selected nodes", EUserInterfaceActionType::Button, FInputChord(EKeys::V) );
diff --git a/Engine/Source/Editor/Sequencer/Private/SequencerSettings.cpp b/Engine/Source/Editor/Sequencer/Private/SequencerSettings.cpp
--- a/Engine/Source/Editor/Sequencer/Private/SequencerSettings.cpp
+++ b/Engine/Source/Editor/Sequencer/Private/SequencerSettings.cpp
@@ -76,6 +76,7 @@ USequencerSettings::USequencerSettings( const FObjectInitializer& ObjectInitiali
 	TrajectoryPathCap = 250;
 	FrameNumberDisplayFormat = EFrameNumberDisplayFormats::Seconds;
 	bAutoExpandNodesOnSelection = true;
+	bAutoExpandOnCurveEditorSelection = true;
 	bRestoreOriginalViewportOnCameraCutUnlock = true;
 	TreeViewWidth = SSequencer::DefaultOutlinerWidth;
 	bShowTickLines = true;
@@ -1152,6 +1153,15 @@ void USequencerSettings::SetAutoExpandNodesOnSelection(bool bInAutoExpandNodesOn
 	}
 }
 
+void USequencerSettings::SetAutoExpandCurveEditorOnSelection(bool bValue)
+{
+	if (bValue != bAutoExpandOnCurveEditorSelection)
+	{
+		bAutoExpandOnCurveEditorSelection = bValue;
+		SaveConfig();
+	}
+}
+
 void USequencerSettings::SetRestoreOriginalViewportOnCameraCutUnlock(bool bInRestoreOriginalViewportOnCameraCutUnlock)
 {
 	if (bInRestoreOriginalViewportOnCameraCutUnlock != bRestoreOriginalViewportOnCameraCutUnlock)
diff --git a/Engine/Source/Editor/Sequencer/Public/SequencerCommands.h b/Engine/Source/Editor/Sequencer/Public/SequencerCommands.h
--- a/Engine/Source/Editor/Sequencer/Public/SequencerCommands.h
+++ b/Engine/Source/Editor/Sequencer/Public/SequencerCommands.h
@@ -164,6 +164,9 @@ public:
 
 	/** Toggle auto expand outliner tree on child selection */
 	TSharedPtr< FUICommandInfo > ToggleAutoExpandNodesOnSelection;
+	
+	/** Toggle auto expand outliner tree on child selection */
+	TSharedPtr< FUICommandInfo > ToggleAutoExpandCurveEditorOnSelection;
 
 	/**
 	 * Toggle whether unlocking a camera cut track should return the viewport to its original location, or keep it where
diff --git a/Engine/Source/Editor/Sequencer/Public/SequencerSettings.h b/Engine/Source/Editor/Sequencer/Public/SequencerSettings.h
--- a/Engine/Source/Editor/Sequencer/Public/SequencerSettings.h
+++ b/Engine/Source/Editor/Sequencer/Public/SequencerSettings.h
@@ -529,7 +529,11 @@ public:
 	/** Sets whether or not to expand the outliner tree view when a child element is selected (from outside of the tree view). */
 	UE_API void SetAutoExpandNodesOnSelection(bool bInAutoExpandNodesOnSelection);
 
-
+	/** Gets whether when selecting items in Sequencer's Outliner will expand Curve Editor's tree view to show the item. */
+	bool GetAutoExpandCurveEditorOnSelection() const { return bAutoExpandOnCurveEditorSelection; }
+	/** Sets whether when selecting items in Sequencer's Outliner will expand Curve Editor's tree view to show the item. */
+	UE_API void SetAutoExpandCurveEditorOnSelection(bool bValue);
+	
 	/**
 	 * Gets whether unlocking a camera cut track should return the viewport to its original location, or keep it where
 	 * the camera cut was.
@@ -764,6 +768,10 @@ protected:
 	UPROPERTY(config, EditAnywhere, Category=CurveEditor)
 	bool bIsolateCurveEditorToSelection;
 	
+	/** When enabled, selecting items in Sequencer's Outliner will expand Curve Editor's tree view to show the item. */
+	UPROPERTY(config, EditAnywhere, Category=CurveEditor)
+	bool bAutoExpandOnCurveEditorSelection;
+	
 	/** Whether the curve editor is visible */
 	UPROPERTY(config, EditAnywhere, Category=CurveEditor)
 	bool bCurveEditorVisible;
