Sequencer: Copy linked filter state into unlinked filter bar when unlinking filters.
- New FLinkedFilterStatePreserver handles the logic.
- This has been added as an option that is enabled by default: FSequencerFilterAreaConfig::bPreserveFiltersOnUnlink
- This extracts command handling to a dedicated FFilterAreaCommandBinder class. This is supposed to prevent explosion of symbols of FFilterAreaManager going forwards.
- SSequencerFilterBar has been adjusted to display the filter widgets in alphabetical display name order. It seems that was the intention before as the TMap was being sorted but the rendering did not match it.
- The reason is so that the order of the unlinked filter bar matches the order the linked filter bar was displaying the widgets in.

#Synced-CL 51881412

diff --git a/Engine/Source/Editor/Sequencer/Private/Filters/Linking/Features/FilterAreaCommandBinder.cpp b/Engine/Source/Editor/Sequencer/Private/Filters/Linking/Features/FilterAreaCommandBinder.cpp
new file mode 100644
--- /dev/null
+++ b/Engine/Source/Editor/Sequencer/Private/Filters/Linking/Features/FilterAreaCommandBinder.cpp
@@ -0,0 +1,91 @@
+﻿// Copyright Epic Games, Inc. All Rights Reserved.
+
+#include "FilterAreaCommandBinder.h"
+
+#include "Filters/SequencerTrackFilterCommands.h"
+#include "Sequencer.h"
+#include "SequencerSettings.h"
+
+namespace UE::Sequencer
+{
+FFilterAreaCommandBinder::FFilterAreaCommandBinder(FName InFilterAreaConfigId, TWeakPtr<FSequencer> InWeakSequencer)
+	: FilterAreaConfigId(InFilterAreaConfigId)
+	, WeakSequencer(MoveTemp(InWeakSequencer))
+{
+	BindCommands();
+}
+
+FFilterAreaCommandBinder::~FFilterAreaCommandBinder()
+{
+	// The command list may be referenced externally.
+	UnbindCommands();
+}
+
+void FFilterAreaCommandBinder::SetFilterBarVisible(bool bIsVisible)
+{
+	const TSharedPtr<FSequencer> SequencerPin = WeakSequencer.Pin();
+	USequencerSettings* Settings = SequencerPin ? SequencerPin->GetSequencerSettings() : nullptr;
+	if (Settings)
+	{
+		constexpr bool bSaveConfig = false; // Saving unnecessary because we'll save right after settings
+		Settings->FindOrAddTrackFilterArea(FilterAreaConfigId, bSaveConfig).SetIsFilterBarVisible(bIsVisible);
+		Settings->SaveConfig();
+	}
+}
+
+bool FFilterAreaCommandBinder::IsFilterBarVisible() const
+{
+	const TSharedPtr<FSequencer> SequencerPin = WeakSequencer.Pin();
+	USequencerSettings* Settings = SequencerPin ? SequencerPin->GetSequencerSettings() : nullptr;
+	
+	constexpr bool bSaveConfig = false; // Saving unnecessary because we're just getting info and using the default if not found.
+	return Settings && Settings->FindOrAddTrackFilterArea(FilterAreaConfigId, bSaveConfig).IsFilterBarVisible();
+}
+
+void FFilterAreaCommandBinder::SetPreserveFiltersOnUnlink(bool bValue)
+{
+	const TSharedPtr<FSequencer> SequencerPin = WeakSequencer.Pin();
+	USequencerSettings* Settings = SequencerPin ? SequencerPin->GetSequencerSettings() : nullptr;
+	if (Settings)
+	{
+		constexpr bool bSaveConfig = false; // Saving unnecessary because we'll save right after settings
+		Settings->FindOrAddTrackFilterArea(FilterAreaConfigId, bSaveConfig).SetPreserveFiltersOnUnlink(bValue);
+		Settings->SaveConfig();
+	}
+}
+
+bool FFilterAreaCommandBinder::GetPreserveFiltersOnUnlink() const
+{
+	const TSharedPtr<FSequencer> SequencerPin = WeakSequencer.Pin();
+	USequencerSettings* Settings = SequencerPin ? SequencerPin->GetSequencerSettings() : nullptr;
+	
+	constexpr bool bSaveConfig = false; // Saving unnecessary because we're just getting info and using the default if not found.
+	return Settings && Settings->FindOrAddTrackFilterArea(FilterAreaConfigId, bSaveConfig).GetPreserveFiltersOnUnlink();
+}
+
+void FFilterAreaCommandBinder::BindCommands()
+{
+	FSequencerTrackFilterCommands& FilterCommands = FSequencerTrackFilterCommands::Get();
+	
+	CommandList->MapAction(
+		FilterCommands.ToggleFilterBarVisibility,
+		FExecuteAction::CreateRaw(this, &FFilterAreaCommandBinder::ToggleFilterBarVisibility),
+		FCanExecuteAction::CreateLambda([]{ return true; }),
+		FIsActionChecked::CreateRaw(this, &FFilterAreaCommandBinder::IsFilterBarVisible)
+		);
+	
+	CommandList->MapAction(
+		FilterCommands.TogglePreserveFiltersOnUnlink,
+		FExecuteAction::CreateRaw(this, &FFilterAreaCommandBinder::TogglePreserveFiltersOnUnlink),
+		FCanExecuteAction::CreateLambda([]{ return true; }),
+		FIsActionChecked::CreateRaw(this, &FFilterAreaCommandBinder::GetPreserveFiltersOnUnlink)
+	);
+}
+
+void FFilterAreaCommandBinder::UnbindCommands()
+{
+	FSequencerTrackFilterCommands& FilterCommands = FSequencerTrackFilterCommands::Get();
+	CommandList->UnmapAction(FilterCommands.ToggleFilterBarVisibility);
+	CommandList->UnmapAction(FilterCommands.TogglePreserveFiltersOnUnlink);
+}
+} // namespace UE::Sequencer
diff --git a/Engine/Source/Editor/Sequencer/Private/Filters/Linking/Features/FilterAreaCommandBinder.h b/Engine/Source/Editor/Sequencer/Private/Filters/Linking/Features/FilterAreaCommandBinder.h
new file mode 100644
--- /dev/null
+++ b/Engine/Source/Editor/Sequencer/Private/Filters/Linking/Features/FilterAreaCommandBinder.h
@@ -0,0 +1,55 @@
+﻿// Copyright Epic Games, Inc. All Rights Reserved.
+
+#pragma once
+
+#include "Framework/Commands/UICommandList.h"
+#include "Templates/SharedPointer.h"
+#include "Templates/UnrealTemplate.h"
+
+class FSequencer;
+
+namespace UE::Sequencer
+{
+/** 
+ * Binds the commands that affect a filter area. 
+ * This is mostly for toggling settings that affect the filter bar regardless of whether it is in linked or unlinked state.
+ */
+class FFilterAreaCommandBinder : public FNoncopyable
+{
+public:
+	
+	explicit FFilterAreaCommandBinder(FName InFilterAreaConfigId, TWeakPtr<FSequencer> InWeakSequencer);
+	~FFilterAreaCommandBinder();
+	
+	/** Sets whether the filter bar pills should be visible. */
+	void SetFilterBarVisible(bool bIsVisible);
+	/** @return Whether the filter bar pills should be visible */
+	bool IsFilterBarVisible() const;
+	/** Inverts the filter bar pill visibility. */
+	void ToggleFilterBarVisibility() { SetFilterBarVisible(!IsFilterBarVisible()); }
+	
+	/** Sets whether the linked filter state is copied into the unlinked filter bar state when switching into unlinked filter state. */
+	void SetPreserveFiltersOnUnlink(bool bValue);
+	/** @return Whether the linked filter state is copied into the unlinked filter bar state when switching into unlinked filter state. */
+	bool GetPreserveFiltersOnUnlink() const;
+	/** Inverts preservation of linked filter state setting. */
+	void TogglePreserveFiltersOnUnlink() { SetPreserveFiltersOnUnlink(!GetPreserveFiltersOnUnlink()); }
+	
+	const TSharedRef<FUICommandList>& GetFilterAreaCommandList() const { return CommandList; }
+	
+private:
+	
+	/** @see USequencerSettings::FindOrAddTrackFilterArea. */
+	const FName FilterAreaConfigId;
+	/** Used to get the settings, e.g. to save filter bar visbility. */
+	const TWeakPtr<FSequencer> WeakSequencer;
+	
+	/** Handles the user's commands, e.g. to change visibility. */
+	const TSharedRef<FUICommandList> CommandList = MakeShared<FUICommandList>();
+	
+	void BindCommands();
+	void UnbindCommands();
+};
+} // namespace UE::Sequencer
+
+
diff --git a/Engine/Source/Editor/Sequencer/Private/Filters/Linking/Features/LinkedFilterStatePreserver.cpp b/Engine/Source/Editor/Sequencer/Private/Filters/Linking/Features/LinkedFilterStatePreserver.cpp
new file mode 100644
--- /dev/null
+++ b/Engine/Source/Editor/Sequencer/Private/Filters/Linking/Features/LinkedFilterStatePreserver.cpp
@@ -0,0 +1,46 @@
+﻿// Copyright Epic Games, Inc. All Rights Reserved.
+
+#include "LinkedFilterStatePreserver.h"
+
+#include "Sequencer.h"
+#include "Filters/Filters/SequencerTrackFilter_CustomText.h"
+#include "Filters/Linking/Mode/LinkedFilterViewModel.h"
+
+namespace UE::Sequencer
+{
+FLinkedFilterStatePreserver::FLinkedFilterStatePreserver(
+	const TSharedRef<FLinkedFilterViewModel>& LinkingModel, 
+	FName FilterAreaConfigId,
+	const TWeakPtr<FSequencer>& WeakSequencer
+	)
+	: LinkingModel(LinkingModel)
+	, FilterAreaConfigId(FilterAreaConfigId)
+	, WeakSequencer(WeakSequencer)
+{
+	LinkingModel->OnFilterModeChanged().AddRaw(this, &FLinkedFilterStatePreserver::OnFilterModeChanged);
+}
+
+FLinkedFilterStatePreserver::~FLinkedFilterStatePreserver()
+{
+	LinkingModel->OnFilterModeChanged().RemoveAll(this);
+}
+
+void FLinkedFilterStatePreserver::OnFilterModeChanged()
+{
+	const TSharedPtr<FSequencer> Sequencer = WeakSequencer.Pin();
+	USequencerSettings* Settings = Sequencer ? Sequencer->GetSequencerSettings() : nullptr;
+	if (!Settings)
+	{
+		return;
+	}
+	
+	FSequencerFilterAreaConfig& Config = Settings->FindOrAddTrackFilterArea(FilterAreaConfigId);
+	const bool bShouldCopyLinkedToUnlinked = LinkingModel->GetFilterMode() == ELinkedFilterMode::Instanced && Config.GetPreserveFiltersOnUnlink();
+	if (bShouldCopyLinkedToUnlinked)
+	{
+		const TSharedRef<FSequencerFilterBar> Source = LinkingModel->GetLinkedFilterBarImpl();
+		const TSharedRef<FSequencerFilterBar> Target = LinkingModel->GetInstancedFilterBarImpl();
+		Target->CopyFiltersFrom(*Source);
+	}
+}
+} // namespace UE::Sequencer
diff --git a/Engine/Source/Editor/Sequencer/Private/Filters/Linking/Features/LinkedFilterStatePreserver.h b/Engine/Source/Editor/Sequencer/Private/Filters/Linking/Features/LinkedFilterStatePreserver.h
new file mode 100644
--- /dev/null
+++ b/Engine/Source/Editor/Sequencer/Private/Filters/Linking/Features/LinkedFilterStatePreserver.h
@@ -0,0 +1,36 @@
+﻿// Copyright Epic Games, Inc. All Rights Reserved.
+
+#pragma once
+
+#include "Templates/SharedPointer.h"
+#include "Templates/UnrealTemplate.h"
+
+class FSequencer;
+
+namespace UE::Sequencer
+{
+class FLinkedFilterViewModel;
+
+/** When the user switched to unlinked mode, this handles copying the filters from the linked state into the unlinked state. */
+class FLinkedFilterStatePreserver : public FNoncopyable
+{
+public:
+	
+	explicit FLinkedFilterStatePreserver(
+		const TSharedRef<FLinkedFilterViewModel>& InLinkingModel, FName InFilterAreaConfigId, const TWeakPtr<FSequencer>& InWeakSequencer
+		);
+	~FLinkedFilterStatePreserver();
+
+private:
+	
+	/** Used to listen when the filter mode changes */
+	const TSharedRef<FLinkedFilterViewModel> LinkingModel;
+	
+	/** @see USequencerSettings::FindOrAddTrackFilterArea. */
+	const FName FilterAreaConfigId;
+	/** Used to get the settings. */
+	const TWeakPtr<FSequencer> WeakSequencer;
+	
+	void OnFilterModeChanged();
+};
+} // namespace UE::Sequencer
diff --git a/Engine/Source/Editor/Sequencer/Private/Filters/Linking/FilterAreaManager.cpp b/Engine/Source/Editor/Sequencer/Private/Filters/Linking/FilterAreaManager.cpp
--- a/Engine/Source/Editor/Sequencer/Private/Filters/Linking/FilterAreaManager.cpp
+++ b/Engine/Source/Editor/Sequencer/Private/Filters/Linking/FilterAreaManager.cpp
@@ -9,46 +9,11 @@ namespace UE::Sequencer
 {
 FFilterAreaManager::FFilterAreaManager(
 	FName InFilterAreaConfigId,
-	TWeakPtr<FSequencer> InWeakSequencer, 
+	const TWeakPtr<FSequencer>& InWeakSequencer, 
 	const TSharedRef<FLinkedFilterViewModel>& InFilterModel
 	)
-	: FilterAreaConfigId(InFilterAreaConfigId)
-	, WeakSequencer(MoveTemp(InWeakSequencer))
-	, FilterModel(InFilterModel)
-{
-	BindCommands();
-}
-
-void FFilterAreaManager::SetFilterBarVisible(bool bIsVisible)
-{
-	const TSharedPtr<FSequencer> SequencerPin = WeakSequencer.Pin();
-	USequencerSettings* Settings = SequencerPin ? SequencerPin->GetSequencerSettings() : nullptr;
-	if (Settings)
-	{
-		constexpr bool bSaveConfig = false; // Saving unnecessary because we'll save right after settings
-		Settings->FindOrAddTrackFilterArea(FilterAreaConfigId, bSaveConfig).SetIsFilterBarVisible(bIsVisible);
-		Settings->SaveConfig();
-	}
-}
-
-bool FFilterAreaManager::IsFilterBarVisible() const
-{
-	const TSharedPtr<FSequencer> SequencerPin = WeakSequencer.Pin();
-	USequencerSettings* Settings = SequencerPin ? SequencerPin->GetSequencerSettings() : nullptr;
-	
-	constexpr bool bSaveConfig = false; // Saving unnecessary because we're just getting info and using the default if not found.
-	return Settings && Settings->FindOrAddTrackFilterArea(FilterAreaConfigId, bSaveConfig).IsFilterBarVisible();
-}
-
-void FFilterAreaManager::BindCommands()
-{
-	FSequencerTrackFilterCommands& FilterCommands = FSequencerTrackFilterCommands::Get();
-	
-	CommandList->MapAction(
-		FilterCommands.ToggleFilterBarVisibility,
-		FExecuteAction::CreateRaw(this, &FFilterAreaManager::ToggleFilterBarVisibility),
-		FCanExecuteAction::CreateLambda([]{ return true; }),
-		FIsActionChecked::CreateRaw(this, &FFilterAreaManager::IsFilterBarVisible)
-		);
-}
+	: FilterModel(InFilterModel)
+	, CommandBinder(InFilterAreaConfigId, InWeakSequencer)
+	, LinkedToUnlinkedCopier(FilterModel, InFilterAreaConfigId, InWeakSequencer)
+{}
 }
diff --git a/Engine/Source/Editor/Sequencer/Private/Filters/Linking/FilterAreaManager.h b/Engine/Source/Editor/Sequencer/Private/Filters/Linking/FilterAreaManager.h
--- a/Engine/Source/Editor/Sequencer/Private/Filters/Linking/FilterAreaManager.h
+++ b/Engine/Source/Editor/Sequencer/Private/Filters/Linking/FilterAreaManager.h
@@ -2,6 +2,8 @@
 
 #pragma once
 
+#include "Features/FilterAreaCommandBinder.h"
+#include "Features/LinkedFilterStatePreserver.h"
 #include "Framework/Commands/UICommandList.h"
 #include "Templates/SharedPointer.h"
 #include "Templates/UnrealTemplate.h"
@@ -23,35 +25,23 @@ public:
 	
 	explicit FFilterAreaManager(
 		FName InFilterAreaConfigId,
-		TWeakPtr<FSequencer> InWeakSequencer, 
+		const TWeakPtr<FSequencer>& InWeakSequencer, 
 		const TSharedRef<FLinkedFilterViewModel>& InFilterModel
 		);
 	
-	/** Sets whether the filter bar pills should be visible. */
-	void SetFilterBarVisible(bool bIsVisible);
-	/** @return Whether the filter bar pills should be visible */
-	bool IsFilterBarVisible() const;
-	/** Inverts the filter bar pill visibility. */
-	void ToggleFilterBarVisibility() { SetFilterBarVisible(!IsFilterBarVisible()); }
-	
 	const TSharedRef<FLinkedFilterViewModel>& GetFilterModel() const { return FilterModel; }
-	const TSharedRef<FUICommandList>& GetFilterAreaCommandList() const { return CommandList; }
+	const TSharedRef<FUICommandList>& GetFilterAreaCommandList() const { return CommandBinder.GetFilterAreaCommandList(); }
 	
 private:
 	
-	/** @see USequencerSettings::FindOrAddTrackFilterArea. */
-	const FName FilterAreaConfigId;
-	
-	/** Used to get the settings, e.g. to save filter bar visbility. */
-	const TWeakPtr<FSequencer> WeakSequencer;
-	
 	/** Manages the linked and unlinked filter state. */
 	const TSharedRef<FLinkedFilterViewModel> FilterModel;
 	
-	/** Handles the user's commands, e.g. to change visibility. */
-	const TSharedRef<FUICommandList> CommandList = MakeShared<FUICommandList>();
+	/** Handles the commands that operate on the filter area on a whole. */
+	FFilterAreaCommandBinder CommandBinder;
 	
-	void BindCommands();
+	/** Copies the filter state from the linked to the unlinked filter bar on switch. */
+	FLinkedFilterStatePreserver LinkedToUnlinkedCopier;
 };
 }
 
diff --git a/Engine/Source/Editor/Sequencer/Private/Filters/Menus/SequencerViewOptionsMenu.cpp b/Engine/Source/Editor/Sequencer/Private/Filters/Menus/SequencerViewOptionsMenu.cpp
--- a/Engine/Source/Editor/Sequencer/Private/Filters/Menus/SequencerViewOptionsMenu.cpp
+++ b/Engine/Source/Editor/Sequencer/Private/Filters/Menus/SequencerViewOptionsMenu.cpp
@@ -11,6 +11,7 @@
 #include "SequencerFilterBarContext.h"
 #include "ToolMenu.h"
 #include "ToolMenus.h"
+#include "Filters/Utils/LinkedFilteringViewUtils.h"
 
 #define LOCTEXT_NAMESPACE "SequencerViewOptionsMenu"
 
@@ -18,13 +19,19 @@ namespace UE::Sequencer
 {
 namespace ViewOptionsDetail
 {
-static void AddFilterVisibilitySection(UToolMenu& InMenu)
+static void AddFilterBarSection(UToolMenu& InMenu)
 {
 	const FSequencerTrackFilterCommands& TrackFilterCommands = FSequencerTrackFilterCommands::Get();
-	FToolMenuSection& FilterBarVisibilitySection = InMenu.FindOrAddSection(
+	FToolMenuSection& FilterBarSection = InMenu.FindOrAddSection(
 		TEXT("FilterBarVisibility") , LOCTEXT("FilterBarVisibilityHeading", "Filter Bar")
 		);
-	FilterBarVisibilitySection.AddMenuEntry(TrackFilterCommands.ToggleFilterBarVisibility);
+	
+	FilterBarSection.AddMenuEntry(TrackFilterCommands.ToggleFilterBarVisibility);
+	
+	if (IsLinkedFilteringFeatureEnabled() && UseLinkedFilteringUI())
+	{
+		FilterBarSection.AddMenuEntry(TrackFilterCommands.TogglePreserveFiltersOnUnlink);
+	}
 }
 
 static TSharedRef<SWidget> CreateMenu(FName InMenuName, FOnPopulateFilterBarMenu InPopulateMenu, const TSharedRef<FUICommandList>& InCommandList)
@@ -94,7 +101,7 @@ void FSequencerViewOptionsMenu::PopulateForCurveEditor(UToolMenu* const InMenu,
     	
 		PopulateFiltersSection(*InMenu);
 		PopulateCurveEditorOrganizeSection(*InMenu);
-		ViewOptionsDetail::AddFilterVisibilitySection(*InMenu);
+		ViewOptionsDetail::AddFilterBarSection(*InMenu);
 	}
 }
 
@@ -205,7 +212,7 @@ void FSequencerViewOptionsMenu::PopulateFilterOptionsSection(UToolMenu& InMenu)
 			EUserInterfaceActionType::ToggleButton);
 	}
 
-	ViewOptionsDetail::AddFilterVisibilitySection(InMenu);
+	ViewOptionsDetail::AddFilterBarSection(InMenu);
 }
 
 void FSequencerViewOptionsMenu::PopulateCurveEditorOrganizeSection(UToolMenu& InMenu)
diff --git a/Engine/Source/Editor/Sequencer/Private/Filters/SequencerFilterAreaConfig.cpp b/Engine/Source/Editor/Sequencer/Private/Filters/SequencerFilterAreaConfig.cpp
--- a/Engine/Source/Editor/Sequencer/Private/Filters/SequencerFilterAreaConfig.cpp
+++ b/Engine/Source/Editor/Sequencer/Private/Filters/SequencerFilterAreaConfig.cpp
@@ -11,3 +11,13 @@ bool FSequencerFilterAreaConfig::IsFilterBarVisible() const
 {
 	return bIsFilterBarVisible;
 }
+
+void FSequencerFilterAreaConfig::SetPreserveFiltersOnUnlink(bool bValue)
+{
+	bPreserveFiltersOnUnlink = bValue;
+}
+
+bool FSequencerFilterAreaConfig::GetPreserveFiltersOnUnlink() const
+{
+	return bPreserveFiltersOnUnlink;
+}
diff --git a/Engine/Source/Editor/Sequencer/Private/Filters/SequencerFilterBar.cpp b/Engine/Source/Editor/Sequencer/Private/Filters/SequencerFilterBar.cpp
--- a/Engine/Source/Editor/Sequencer/Private/Filters/SequencerFilterBar.cpp
+++ b/Engine/Source/Editor/Sequencer/Private/Filters/SequencerFilterBar.cpp
@@ -29,6 +29,7 @@
 #include "MVVM/Views/SOutlinerView.h"
 #include "Sequencer.h"
 #include "SSequencer.h"
+#include "Algo/ForEach.h"
 #include "Filters/SequencerFilterAreaConfig.h"
 #include "Stats/StatsMisc.h"
 #include "UObject/UObjectIterator.h"
@@ -102,6 +103,55 @@ TSharedPtr<ICustomTextFilter<FSequencerTrackFilterType>> FSequencerFilterBar::Cr
 	return MakeShared<FSequencerTrackFilter_CustomText>(*this);
 }
 
+void FSequencerFilterBar::CopyFiltersFrom(const FSequencerFilterBar& InSource)
+{
+	USequencerSettings* const SequencerSettings = Sequencer.GetSequencerSettings();
+	if (!SequencerSettings)
+	{
+		return;
+	}
+	
+	// Needed to notify UI to update.
+	const auto BroadcastDeactivate = [this](const TSharedRef<FSequencerTrackFilter>& InFilter)
+	{
+		FiltersChangedEvent.Broadcast(ESequencerFilterChange::Disable, InFilter);
+		return true;
+	};
+	CommonFilters->ForEachFilter(BroadcastDeactivate);
+	InternalFilters->ForEachFilter(BroadcastDeactivate);
+	Algo::ForEach(CustomTextFilters, BroadcastDeactivate);
+	
+	// We'll be lazy and simply override the config...
+	FSequencerFilterBarConfig& SourceSettings = SequencerSettings->FindOrAddTrackFilterBar(InSource.GetIdentifier(), false);
+	FSequencerFilterBarConfig& TargetSettings = SequencerSettings->FindOrAddTrackFilterBar(GetIdentifier(), false);
+	TargetSettings = SourceSettings;
+	SequencerSettings->SaveConfig();
+	
+	// ... and reuse the initialization mechanisms to recreate filters from config.
+	CreateDefaultFilters();
+	CreateCustomTextFiltersFromConfig();
+	
+	// Needed to notify UI to update.
+	const auto BroadcastFilter = [this](const TSharedRef<FSequencerTrackFilter>& InFilter)
+	{
+		if (IsFilterActive(InFilter))
+		{
+			FiltersChangedEvent.Broadcast(ESequencerFilterChange::Activate, InFilter);
+		}
+		if (IsFilterEnabled(InFilter))
+		{
+			FiltersChangedEvent.Broadcast(ESequencerFilterChange::Enable, InFilter);
+		}
+		return true;
+	};
+	CommonFilters->ForEachFilter(BroadcastFilter);
+	InternalFilters->ForEachFilter(BroadcastFilter);
+	Algo::ForEach(CustomTextFilters, BroadcastFilter);
+	
+	// The filters are now dirty
+	RequestFilterUpdate();
+}
+
 void FSequencerFilterBar::CreateDefaultFilters()
 {
 	auto AddFilterIfSupported = [this]
@@ -1443,6 +1493,8 @@ FString FSequencerFilterBar::GenerateTextFilterStringFromEnabledFilters() const
 
 bool FSequencerFilterBar::FilterNodesRecursive(const bool bInHasActiveFilter, const TViewModelPtr<IOutlinerExtension>& InStartNode)
 {
+	TRACE_CPUPROFILER_EVENT_SCOPE(FSequencerFilterBar::FilterNodesRecursive);
+	
 	/**
 	 * Main Filtering Logic
 	 *
diff --git a/Engine/Source/Editor/Sequencer/Private/Filters/SequencerFilterBar.h b/Engine/Source/Editor/Sequencer/Private/Filters/SequencerFilterBar.h
--- a/Engine/Source/Editor/Sequencer/Private/Filters/SequencerFilterBar.h
+++ b/Engine/Source/Editor/Sequencer/Private/Filters/SequencerFilterBar.h
@@ -31,9 +31,13 @@ namespace UE::Sequencer { struct FFilterBarInitArgs; }
 
 enum class ESequencerFilterChange : uint8
 {
+	/** The filter shows up in the filter bar. In the combo menu, it is checked. */
 	Enable,
+	/** The filter does not show up in the filter bar. In the combo menu, it is unchecked. */
 	Disable,
+	/** The filter shows up in the filter bar and is currently filtering. In the combo menu, it is checked. */
 	Activate,
+	/** The filter shows up in the filter bar but is currently not filtering. In the combo menu, it is checked. */
 	Deactivate
 };
 
@@ -51,6 +55,13 @@ public:
 	void BindCommands();
 
 	TSharedPtr<ICustomTextFilter<FSequencerTrackFilterType>> CreateTextFilter();
+	
+	/**
+	 * Copies the activation and enabled state of all filters from InSource
+	 * to this filter bar. After the call, both instances contain identical
+	 * filter states.
+	 */
+	void CopyFiltersFrom(const FSequencerFilterBar& InSource);
 
 	virtual bool AreFiltersMuted() const override;
 	virtual void MuteFilters(const bool bInMute) override;
diff --git a/Engine/Source/Editor/Sequencer/Private/Filters/SequencerTrackFilterCommands.cpp b/Engine/Source/Editor/Sequencer/Private/Filters/SequencerTrackFilterCommands.cpp
--- a/Engine/Source/Editor/Sequencer/Private/Filters/SequencerTrackFilterCommands.cpp
+++ b/Engine/Source/Editor/Sequencer/Private/Filters/SequencerTrackFilterCommands.cpp
@@ -29,7 +29,8 @@ void FSequencerTrackFilterCommands::RegisterCommands()
 	UI_COMMAND(DeactivateAllFilters, "Deactivate All Filters", "Deactivates all enabled filters", EUserInterfaceActionType::Button, FInputChord());
 	
 	UI_COMMAND(ToggleLinkedFiltering, "Toggle Linked Filtering", "Toggles whether filtering should use the filter state shared with other UI or its own instance.", EUserInterfaceActionType::ToggleButton, FInputChord());
-
+	UI_COMMAND(TogglePreserveFiltersOnUnlink, "Preserve Filters on Unlink", "Toggles whether when you unlink your filters, you want the filters in your linked filter bar to be copied into the unlinked filter bar.", EUserInterfaceActionType::ToggleButton, FInputChord())
+	
 	UI_COMMAND(HideSelectedTracks, "Hide Selected Tracks", "Hide all selected tracks", EUserInterfaceActionType::Button, FInputChord(EModifierKey::Alt, EKeys::W));
 	UI_COMMAND(IsolateSelectedTracks, "Isolate Selected Tracks", "Isolate all selected tracks", EUserInterfaceActionType::Button, FInputChord(EModifierKey::Alt, EKeys::Q));
 
diff --git a/Engine/Source/Editor/Sequencer/Private/Filters/SequencerTrackFilterCommands.h b/Engine/Source/Editor/Sequencer/Private/Filters/SequencerTrackFilterCommands.h
--- a/Engine/Source/Editor/Sequencer/Private/Filters/SequencerTrackFilterCommands.h
+++ b/Engine/Source/Editor/Sequencer/Private/Filters/SequencerTrackFilterCommands.h
@@ -33,6 +33,7 @@ public:
 	TSharedPtr<FUICommandInfo> DeactivateAllFilters;
 	
 	TSharedPtr<FUICommandInfo> ToggleLinkedFiltering;
+	TSharedPtr<FUICommandInfo> TogglePreserveFiltersOnUnlink;
 
 	/** Hide/Isolate Filter Commands */
 
diff --git a/Engine/Source/Editor/Sequencer/Private/Filters/Widgets/SSequencerFilterBar.cpp b/Engine/Source/Editor/Sequencer/Private/Filters/Widgets/SSequencerFilterBar.cpp
--- a/Engine/Source/Editor/Sequencer/Private/Filters/Widgets/SSequencerFilterBar.cpp
+++ b/Engine/Source/Editor/Sequencer/Private/Filters/Widgets/SSequencerFilterBar.cpp
@@ -229,7 +229,12 @@ void SSequencerFilterBar::AddFilterWidget(const TSharedRef<FSequencerTrackFilter
 			return InFilterA->GetDisplayName().CompareTo(InFilterB->GetDisplayName()) < 0;
 		});
 
-	AddWidgetToLayout(InFilterWidget);
+	// Make sure the widgets are rendered in alphabetical order.
+	HorizontalContainerWidget->ClearChildren();
+	for (const TPair<TSharedRef<FSequencerTrackFilter>, TSharedRef<SSequencerFilter>>& Pair : FilterWidgets)
+	{
+		AddWidgetToLayout(Pair.Value);
+	}
 }
 
 void SSequencerFilterBar::RemoveFilterWidget(const TSharedRef<FSequencerTrackFilter>& InFilter)
diff --git a/Engine/Source/Editor/Sequencer/Public/Filters/SequencerFilterAreaConfig.h b/Engine/Source/Editor/Sequencer/Public/Filters/SequencerFilterAreaConfig.h
--- a/Engine/Source/Editor/Sequencer/Public/Filters/SequencerFilterAreaConfig.h
+++ b/Engine/Source/Editor/Sequencer/Public/Filters/SequencerFilterAreaConfig.h
@@ -18,12 +18,19 @@ struct FSequencerFilterAreaConfig
 	UE_API void SetIsFilterBarVisible(bool bNewIsVisible);
 	UE_API bool IsFilterBarVisible() const;
 	
+	UE_API void SetPreserveFiltersOnUnlink(bool bValue);
+	UE_API bool GetPreserveFiltersOnUnlink() const;
+	
 private:
 	
 	/** Whether the filter bar should be visible. */
 	UPROPERTY()
 	bool bIsFilterBarVisible = false;
 	
+	/** Whether when you unlink your filters, you want the filters in your linked filter bar to be copied into the unlinked filter bar. */
+	UPROPERTY()
+	bool bPreserveFiltersOnUnlink = true;
+	
 	// TODO UE-362150: safe whether the filter bar should be instanced or linked
 };
 
