Sequencer: Fix view options missing binding to filter command list.

- FFilterAreaManager::GetFilterAreaCommandList now stores all commands that affect filters or linked filtering.
- All users of FFilterAreaManager now ensure that the command list is in their input stack.
- FFilterAreaCommandBinder handles dynamically rebinding to the active filter bar instance when the filter mode changes.
- FSequencerFilterBar now handles calling BindCommands by itself ensuring that callers cannot forget calling it. For this, a static Make function needed to be introduced as BindCommands makes use of CreateSP, which requires TSharedFromThis to have been initialized.

#Synced-CL 51881412

diff --git a/Engine/Source/Editor/Sequencer/Private/CurveEditor/Integration/CurveEditorCommandBinder.cpp b/Engine/Source/Editor/Sequencer/Private/CurveEditor/Integration/CurveEditorCommandBinder.cpp
--- a/Engine/Source/Editor/Sequencer/Private/CurveEditor/Integration/CurveEditorCommandBinder.cpp
+++ b/Engine/Source/Editor/Sequencer/Private/CurveEditor/Integration/CurveEditorCommandBinder.cpp
@@ -72,7 +72,9 @@ void FCurveCommandBinder::InitSequencerCommandBindings()
 	
 	CurveEditorSharedBindings->MapAction(Commands.ToggleAutoExpandCurveEditorOnSelection, *SequencerCommandBindings->GetActionForCommand(Commands.ToggleAutoExpandCurveEditorOnSelection));
 	
-	FilterArea.GetFilterAreaCommandList()->Append(CurveEditorSharedBindings.ToSharedRef());
+	// The shared bindings should include the commands related to filtering.
+	CurveEditorSharedBindings.ToSharedRef()->Append(FilterArea.GetFilterAreaCommandList());
+	// Curve Editor's command list should execute whatever we bind in Sequencer.
 	CurveEditor->GetCommands()->Append(CurveEditorSharedBindings.ToSharedRef());
 	
 	BindCurveEditorToActiveFilterBar();
diff --git a/Engine/Source/Editor/Sequencer/Private/Filters/Linking/Features/FilterAreaCommandBinder.cpp b/Engine/Source/Editor/Sequencer/Private/Filters/Linking/Features/FilterAreaCommandBinder.cpp
--- a/Engine/Source/Editor/Sequencer/Private/Filters/Linking/Features/FilterAreaCommandBinder.cpp
+++ b/Engine/Source/Editor/Sequencer/Private/Filters/Linking/Features/FilterAreaCommandBinder.cpp
@@ -2,23 +2,30 @@
 
 #include "FilterAreaCommandBinder.h"
 
+#include "Filters/Linking/Mode/LinkedFilterViewModel.h"
 #include "Filters/SequencerTrackFilterCommands.h"
+#include "Filters/Utils/LinkedFilteringViewUtils.h"
 #include "Sequencer.h"
 #include "SequencerSettings.h"
 
 namespace UE::Sequencer
 {
-FFilterAreaCommandBinder::FFilterAreaCommandBinder(FName InFilterAreaConfigId, TWeakPtr<FSequencer> InWeakSequencer)
+FFilterAreaCommandBinder::FFilterAreaCommandBinder(
+	FName InFilterAreaConfigId, TWeakPtr<FSequencer> InWeakSequencer, const TSharedRef<FLinkedFilterViewModel>& InLinkedFilterModel
+	)
 	: FilterAreaConfigId(InFilterAreaConfigId)
 	, WeakSequencer(MoveTemp(InWeakSequencer))
+	, LinkedFilterModel(InLinkedFilterModel)
 {
 	BindCommands();
+	LinkedFilterModel->OnFilterModeChanged().AddRaw(this, &FFilterAreaCommandBinder::OnFilterModeChanged);
 }
 
 FFilterAreaCommandBinder::~FFilterAreaCommandBinder()
 {
 	// The command list may be referenced externally.
 	UnbindCommands();
+	LinkedFilterModel->OnFilterModeChanged().RemoveAll(this);
 }
 
 void FFilterAreaCommandBinder::SetFilterBarVisible(bool bIsVisible)
@@ -67,25 +74,44 @@ void FFilterAreaCommandBinder::BindCommands()
 {
 	FSequencerTrackFilterCommands& FilterCommands = FSequencerTrackFilterCommands::Get();
 	
-	CommandList->MapAction(
+	RootCommandList->MapAction(
+		FSequencerTrackFilterCommands::Get().ToggleLinkedFiltering ,
+		FExecuteAction::CreateLambda([WeakViewModel = LinkedFilterModel.ToWeakPtr()] { ToggleLinkedFiltering(WeakViewModel); }),
+		FCanExecuteAction::CreateLambda([]{ return IsLinkedFilteringFeatureEnabled(); }),
+		FIsActionChecked::CreateLambda([WeakViewModel = LinkedFilterModel.ToWeakPtr()] { return IsLinkedFilteringActionChecked(WeakViewModel); })
+		);
+	
+	RootCommandList->MapAction(
 		FilterCommands.ToggleFilterBarVisibility,
 		FExecuteAction::CreateRaw(this, &FFilterAreaCommandBinder::ToggleFilterBarVisibility),
 		FCanExecuteAction::CreateLambda([]{ return true; }),
 		FIsActionChecked::CreateRaw(this, &FFilterAreaCommandBinder::IsFilterBarVisible)
 		);
 	
-	CommandList->MapAction(
+	RootCommandList->MapAction(
 		FilterCommands.TogglePreserveFiltersOnUnlink,
 		FExecuteAction::CreateRaw(this, &FFilterAreaCommandBinder::TogglePreserveFiltersOnUnlink),
 		FCanExecuteAction::CreateLambda([]{ return true; }),
 		FIsActionChecked::CreateRaw(this, &FFilterAreaCommandBinder::GetPreserveFiltersOnUnlink)
 	);
+	
+	RebindActiveFilterBarCommands();
 }
 
 void FFilterAreaCommandBinder::UnbindCommands()
 {
 	FSequencerTrackFilterCommands& FilterCommands = FSequencerTrackFilterCommands::Get();
-	CommandList->UnmapAction(FilterCommands.ToggleFilterBarVisibility);
-	CommandList->UnmapAction(FilterCommands.TogglePreserveFiltersOnUnlink);
+	RootCommandList->UnmapAction(FilterCommands.ToggleLinkedFiltering);
+	RootCommandList->UnmapAction(FilterCommands.ToggleFilterBarVisibility);
+	RootCommandList->UnmapAction(FilterCommands.TogglePreserveFiltersOnUnlink);
+}
+
+void FFilterAreaCommandBinder::RebindActiveFilterBarCommands()
+{
+	ensure(ActiveFilterBarCommandList.GetSharedReferenceCount() <= 1);
+	ActiveFilterBarCommandList = MakeShared<FUICommandList>(*LinkedFilterModel->GetActiveFilterBarImpl()->GetCommandList());
+	
+	// Append will clean up the previous, now-dead reference to ActiveFilterBarCommandList.
+	RootCommandList->Append(ActiveFilterBarCommandList.ToSharedRef());
 }
 } // namespace UE::Sequencer
diff --git a/Engine/Source/Editor/Sequencer/Private/Filters/Linking/Features/FilterAreaCommandBinder.h b/Engine/Source/Editor/Sequencer/Private/Filters/Linking/Features/FilterAreaCommandBinder.h
--- a/Engine/Source/Editor/Sequencer/Private/Filters/Linking/Features/FilterAreaCommandBinder.h
+++ b/Engine/Source/Editor/Sequencer/Private/Filters/Linking/Features/FilterAreaCommandBinder.h
@@ -10,15 +10,20 @@ class FSequencer;
 
 namespace UE::Sequencer
 {
-/** 
- * Binds the commands that affect a filter area. 
- * This is mostly for toggling settings that affect the filter bar regardless of whether it is in linked or unlinked state.
+class FLinkedFilterViewModel;
+
+/**
+ * Manages UI command bindings for a filter area.
+ * - Binds persistent commands that apply regardless of linked/unlinked filter mode (e.g. toggling filter bar visibility and preserve-on-unlink).
+ * - Dynamically rebinds RootCommandList to the active filter bar instance whenever ELinkedFilterMode changes.
  */
 class FFilterAreaCommandBinder : public FNoncopyable
 {
 public:
 	
-	explicit FFilterAreaCommandBinder(FName InFilterAreaConfigId, TWeakPtr<FSequencer> InWeakSequencer);
+	explicit FFilterAreaCommandBinder(
+		FName InFilterAreaConfigId, TWeakPtr<FSequencer> InWeakSequencer, const TSharedRef<FLinkedFilterViewModel>& InLinkedFilterModel
+		);
 	~FFilterAreaCommandBinder();
 	
 	/** Sets whether the filter bar pills should be visible. */
@@ -35,7 +40,8 @@ public:
 	/** Inverts preservation of linked filter state setting. */
 	void TogglePreserveFiltersOnUnlink() { SetPreserveFiltersOnUnlink(!GetPreserveFiltersOnUnlink()); }
 	
-	const TSharedRef<FUICommandList>& GetFilterAreaCommandList() const { return CommandList; }
+	/** Handles all commands related to the filter bar. The system using FFilterAreaManager should hook this command list up. */
+	const TSharedRef<FUICommandList>& GetFilterAreaCommandList() const { return RootCommandList; }
 	
 private:
 	
@@ -44,12 +50,20 @@ private:
 	/** Used to get the settings, e.g. to save filter bar visbility. */
 	const TWeakPtr<FSequencer> WeakSequencer;
 	
-	/** Handles the user's commands, e.g. to change visibility. */
-	const TSharedRef<FUICommandList> CommandList = MakeShared<FUICommandList>();
+	/** Used to implement some of the commands. */
+	const TSharedRef<FLinkedFilterViewModel> LinkedFilterModel;
 	
+	/** Handles all commands related to the filter bar. */
+	const TSharedRef<FUICommandList> RootCommandList = MakeShared<FUICommandList>();
+	/** The active filter bar's commands are bound to this command list. They are rebound every time the linked filtering mode changes. */
+	TSharedPtr<FUICommandList> ActiveFilterBarCommandList;
+
 	void BindCommands();
 	void UnbindCommands();
+	
+	/** Responds to linked filtering mode changing */
+	void OnFilterModeChanged() { RebindActiveFilterBarCommands(); }
+	/** Rebinds the actions mapped in ActiveFilterBarCommandList to the active filter bar. */
+	void RebindActiveFilterBarCommands();
 };
 } // namespace UE::Sequencer
-
-
diff --git a/Engine/Source/Editor/Sequencer/Private/Filters/Linking/FilterAreaManager.cpp b/Engine/Source/Editor/Sequencer/Private/Filters/Linking/FilterAreaManager.cpp
--- a/Engine/Source/Editor/Sequencer/Private/Filters/Linking/FilterAreaManager.cpp
+++ b/Engine/Source/Editor/Sequencer/Private/Filters/Linking/FilterAreaManager.cpp
@@ -13,7 +13,7 @@ FFilterAreaManager::FFilterAreaManager(
 	const TSharedRef<FLinkedFilterViewModel>& InFilterModel
 	)
 	: FilterModel(InFilterModel)
-	, CommandBinder(InFilterAreaConfigId, InWeakSequencer)
+	, CommandBinder(InFilterAreaConfigId, InWeakSequencer, FilterModel)
 	, LinkedToUnlinkedCopier(FilterModel, InFilterAreaConfigId, InWeakSequencer)
 {}
 }
diff --git a/Engine/Source/Editor/Sequencer/Private/Filters/Linking/FilterAreaManager.h b/Engine/Source/Editor/Sequencer/Private/Filters/Linking/FilterAreaManager.h
--- a/Engine/Source/Editor/Sequencer/Private/Filters/Linking/FilterAreaManager.h
+++ b/Engine/Source/Editor/Sequencer/Private/Filters/Linking/FilterAreaManager.h
@@ -30,6 +30,11 @@ public:
 		);
 	
 	const TSharedRef<FLinkedFilterViewModel>& GetFilterModel() const { return FilterModel; }
+	
+	/** 
+	 * @return Command list mapping all commands related to the filter bar and their filters. 
+	 * The owner of FFilterAreaManager should hook this up to their input handler.
+	 */
 	const TSharedRef<FUICommandList>& GetFilterAreaCommandList() const { return CommandBinder.GetFilterAreaCommandList(); }
 	
 private:
diff --git a/Engine/Source/Editor/Sequencer/Private/Filters/Linking/Mode/LinkedFilterFactoryViewModel.cpp b/Engine/Source/Editor/Sequencer/Private/Filters/Linking/Mode/LinkedFilterFactoryViewModel.cpp
--- a/Engine/Source/Editor/Sequencer/Private/Filters/Linking/Mode/LinkedFilterFactoryViewModel.cpp
+++ b/Engine/Source/Editor/Sequencer/Private/Filters/Linking/Mode/LinkedFilterFactoryViewModel.cpp
@@ -19,10 +19,9 @@ FLinkedFilterFactoryViewModel::FLinkedFilterFactoryViewModel(
 
 TSharedRef<FLinkedFilterViewModel> FLinkedFilterFactoryViewModel::MakeFilteringModelImpl(const FMakeLinkedFilterViewModelArgs& InArgs) const
 {
-	const TSharedRef<FSequencerFilterBar> InstancedFilterBar = MakeShared<FSequencerFilterBar>(
+	const TSharedRef<FSequencerFilterBar> InstancedFilterBar = FSequencerFilterBar::Make(
 		FFilterBarInitArgs(Sequencer, InArgs.ConfigSubIdentifier, InArgs.UnlinkedFilterFlags)
 		);
-	InstancedFilterBar->BindCommands();
 	return MakeShared<FLinkedFilterViewModel>(SharedFilterBar, InstancedFilterBar, FilterEvaluator);
 }
 
diff --git a/Engine/Source/Editor/Sequencer/Private/Filters/SequencerFilterBar.cpp b/Engine/Source/Editor/Sequencer/Private/Filters/SequencerFilterBar.cpp
--- a/Engine/Source/Editor/Sequencer/Private/Filters/SequencerFilterBar.cpp
+++ b/Engine/Source/Editor/Sequencer/Private/Filters/SequencerFilterBar.cpp
@@ -41,7 +41,17 @@ using namespace UE::Sequencer;
 
 int32 FSequencerFilterBar::InstanceCount = 0;
 
-FSequencerFilterBar::FSequencerFilterBar(FFilterBarInitArgs InArgs)
+TSharedRef<FSequencerFilterBar> FSequencerFilterBar::Make(FFilterBarInitArgs InArgs)
+{
+	TSharedRef<FSequencerFilterBar> FilterBar = MakeShared<FSequencerFilterBar>(MoveTemp(InArgs), FPrivateToken());
+	
+	// Requires TSharedFromThis to be initialized.
+	FilterBar->BindCommands(); 
+	
+	return FilterBar;
+}
+
+FSequencerFilterBar::FSequencerFilterBar(FFilterBarInitArgs InArgs, FPrivateToken)
 	: Sequencer(*InArgs.Sequencer)
 	, CommandList(MakeShared<FUICommandList>())
 	, ClassTypeCategory(MakeShared<FFilterCategory>(LOCTEXT("ActorTypeFilterCategory", "Actor Type Filters"), FText()))
diff --git a/Engine/Source/Editor/Sequencer/Private/Filters/SequencerFilterBar.h b/Engine/Source/Editor/Sequencer/Private/Filters/SequencerFilterBar.h
--- a/Engine/Source/Editor/Sequencer/Private/Filters/SequencerFilterBar.h
+++ b/Engine/Source/Editor/Sequencer/Private/Filters/SequencerFilterBar.h
@@ -44,16 +44,18 @@ enum class ESequencerFilterChange : uint8
 /** Holds the Sequencer track filter collection, the current text filter, and hidden/isolated lists. */
 class FSequencerFilterBar : public ISequencerTrackFilters
 {
+	struct FPrivateToken { explicit FPrivateToken() = default; }; // Force construction through FSequencerFilterBar::Make.
 public:
 	
 	DECLARE_EVENT_TwoParams(FSequencerFilterBar, FSequencerFiltersChanged, const ESequencerFilterChange /*InChangeType*/, const TSharedRef<FSequencerTrackFilter>& /*InFilter*/);
 	DECLARE_EVENT_TwoParams(FSequencerFilterBar, FSequencerCustomTextFiltersChanged, const ESequencerFilterChange /*InChangeType*/, const TSharedRef<FSequencerTrackFilter_CustomText>& /*InFilter*/);
 	
-	explicit FSequencerFilterBar(UE::Sequencer::FFilterBarInitArgs InArgs);
+	/** Creates a filter bar instance. */
+	static TSharedRef<FSequencerFilterBar> Make(UE::Sequencer::FFilterBarInitArgs InArgs);
+	
+	explicit FSequencerFilterBar(UE::Sequencer::FFilterBarInitArgs InArgs, FPrivateToken);
 	virtual ~FSequencerFilterBar() override;
 
-	void BindCommands();
-
 	TSharedPtr<ICustomTextFilter<FSequencerTrackFilterType>> CreateTextFilter();
 	
 	/**
@@ -313,6 +315,9 @@ private:
 	/** @return Gets the selected IOutlinerExtensions, or all of them if nothing is selected. */
 	TSet<UE::Sequencer::TWeakViewModelPtr<UE::Sequencer::IOutlinerExtension>> GetSelectedTracksOrAll() const;
 	
+	/** Binds all commands that the filter bar or its filters handle. */
+	void BindCommands();
+	
 	/** Do not call directly! Should only be called by FilterNodes(). */
 	bool FilterNodesRecursive(const bool bInHasActiveFilter, const UE::Sequencer::TViewModelPtr<UE::Sequencer::IOutlinerExtension>& InStartNode);
 	
diff --git a/Engine/Source/Editor/Sequencer/Private/MVVM/ViewModels/LinkedFilteringMediator.cpp b/Engine/Source/Editor/Sequencer/Private/MVVM/ViewModels/LinkedFilteringMediator.cpp
--- a/Engine/Source/Editor/Sequencer/Private/MVVM/ViewModels/LinkedFilteringMediator.cpp
+++ b/Engine/Source/Editor/Sequencer/Private/MVVM/ViewModels/LinkedFilteringMediator.cpp
@@ -39,33 +39,9 @@ FLinkedFilteringMediator::~FLinkedFilteringMediator()
 
 void FLinkedFilteringMediator::BindCommands(const TSharedRef<FUICommandList>& InSequencerCommandList) const
 {
-	const TSharedRef<FSequencerFilterBar> SharedFilterBar = GetFilterModel()->GetLinkedFilterBarImpl();
-	SharedFilterBar->BindCommands();
-	
-	const TSharedRef<FSequencerFilterBar> InstancedFilterBar = GetFilterModel()->GetInstancedFilterBarImpl();
-	InstancedFilterBar->BindCommands();
-	
-	InSequencerCommandList->MapAction(
-		FSequencerTrackFilterCommands::Get().ToggleLinkedFiltering ,
-		FExecuteAction::CreateLambda([WeakViewModel = GetFilterModel().ToWeakPtr()] { ToggleLinkedFiltering(WeakViewModel); }),
-		FCanExecuteAction::CreateLambda([]{ return IsLinkedFilteringFeatureEnabled(); }),
-		FIsActionChecked::CreateLambda([WeakViewModel = GetFilterModel().ToWeakPtr()] { return IsLinkedFilteringActionChecked(WeakViewModel); })
-		);
-	
 	InSequencerCommandList->Append(GetFilterAreaCommandList());
 }
 
-bool FLinkedFilteringMediator::ProcessCommandBindings(const FKeyEvent& InKeyEvent) const
-{
-	if (GetFilterAreaCommandList()->ProcessCommandBindings(InKeyEvent))
-	{
-		return true;
-	}
-	
-	const TSharedPtr<FUICommandList> ActiveFilterCommandList = GetFilterModel()->GetActiveFilterBar()->GetCommandList();
-	return ActiveFilterCommandList && ActiveFilterCommandList->ProcessCommandBindings(InKeyEvent);
-}
-
 ELinkedFilterMode FLinkedFilteringMediator::GetFilterMode() const
 {
 	return GetFilterModel()->GetFilterMode();
diff --git a/Engine/Source/Editor/Sequencer/Private/MVVM/ViewModels/LinkedFilteringMediator.h b/Engine/Source/Editor/Sequencer/Private/MVVM/ViewModels/LinkedFilteringMediator.h
--- a/Engine/Source/Editor/Sequencer/Private/MVVM/ViewModels/LinkedFilteringMediator.h
+++ b/Engine/Source/Editor/Sequencer/Private/MVVM/ViewModels/LinkedFilteringMediator.h
@@ -31,9 +31,6 @@ public:
 	/** Binds the filter bar commands. */
 	void BindCommands(const TSharedRef<FUICommandList>& InSequencerCommandList) const;
 	
-	/** Process the commands of the active filter bar. */
-	bool ProcessCommandBindings(const FKeyEvent& InKeyEvent) const;
-	
 private:
 	
 	/** @return The filter mode that Sequencer is currently using. */
diff --git a/Engine/Source/Editor/Sequencer/Private/MVVM/ViewModels/SequencerEditorViewModelImpl.cpp b/Engine/Source/Editor/Sequencer/Private/MVVM/ViewModels/SequencerEditorViewModelImpl.cpp
--- a/Engine/Source/Editor/Sequencer/Private/MVVM/ViewModels/SequencerEditorViewModelImpl.cpp
+++ b/Engine/Source/Editor/Sequencer/Private/MVVM/ViewModels/SequencerEditorViewModelImpl.cpp
@@ -36,11 +36,10 @@ void FSequencerEditorViewModelImpl::InitLinkedFiltering()
 	const TCHAR* ConfigId = ConfigIds::LinkedConfigSubId;
 	// Linked filter state should be affected by the USequencerSettings::bIncludePinnedInFilter setting.
 	constexpr EFilterFlags LinkedFilterFlags = EFilterFlags::PinnedItemsCanSkipFiltering;
+	const TSharedRef<FSequencerFilterBar> LinkedFilterBar = FSequencerFilterBar::Make(FFilterBarInitArgs(*Sequencer, ConfigId, LinkedFilterFlags));
 	
 	LinkedFilterFactory = MakeShared<FLinkedFilterFactoryViewModel>(
-		*Sequencer,
-		MakeShared<FSequencerFilterBar>(FFilterBarInitArgs(*Sequencer, ConfigId, LinkedFilterFlags)),
-		Sequencer->GetNodeTree()->GetFilterEvaluator()
+		*Sequencer, LinkedFilterBar, Sequencer->GetNodeTree()->GetFilterEvaluator()
 		);
 	
 	// The unlinked filter state for the Outliner should also be affected by the USequencerSettings::bIncludePinnedInFilter setting.
diff --git a/Engine/Source/Editor/Sequencer/Private/SSequencer.cpp b/Engine/Source/Editor/Sequencer/Private/SSequencer.cpp
--- a/Engine/Source/Editor/Sequencer/Private/SSequencer.cpp
+++ b/Engine/Source/Editor/Sequencer/Private/SSequencer.cpp
@@ -2900,13 +2900,6 @@ FReply SSequencer::OnKeyDown( const FGeometry& MyGeometry, const FKeyEvent& InKe
 		return FReply::Handled();
 	}
 
-	using namespace UE::Sequencer;
-	const TSharedPtr<FLinkedFilteringMediator> FilterMediator = SequencerPin->GetViewModelImpl()->GetFilterMediator();
-	if (FilterMediator && FilterMediator->ProcessCommandBindings(InKeyEvent))
-	{
-		return FReply::Handled();
-	}
-
 	return FReply::Unhandled();
 }
 
