Sequencer: Fixed muting a single mix layer disabling evaluation for the entire mixer track

- The mute-to-eval-disabled conversion loop in SSequencer iterated ITrackExtension
nodes and called SetLocalEvalDisabled on the track. When multiple nodes (e.g. mixer
layers) shared the same track, the last one processed won, so muting the last layer
disabled the entire track
- Skip ITrackExtension nodes that also implement ITrackRowExtension in the track-level
loop, deferring to the per-row loop which correctly uses SetLocalRowEvalDisabled

#Synced-CL 51964473

diff --git a/Engine/Source/Editor/Sequencer/Private/SSequencer.cpp b/Engine/Source/Editor/Sequencer/Private/SSequencer.cpp
--- a/Engine/Source/Editor/Sequencer/Private/SSequencer.cpp
+++ b/Engine/Source/Editor/Sequencer/Private/SSequencer.cpp
@@ -2638,6 +2638,16 @@ void SSequencer::UpdateLayoutTree()
 
 			for (TViewModelPtr<ITrackExtension> TrackNode : Sequencer->GetNodeTree()->GetRootNode()->GetDescendantsOfType<ITrackExtension>())
 			{
+				// Skip nodes that also implement ITrackRowExtension. These represent
+				// individual rows within a shared track (e.g. Animation Mixer layers)
+				// and will be handled by the per-row loop below. Applying track-level
+				// SetLocalEvalDisabled here would disable the entire shared track
+				// based on a single row's mute state.
+				if (TrackNode.AsModel()->IsA<ITrackRowExtension>())
+				{
+					continue;
+				}
+
 				UMovieSceneTrack* const Track = TrackNode->GetTrack();
 				if (!IsValid(Track))
 				{
