Sequencer: Added per-section mute/solo eval-disable and transition auto-disable.

This is necessary now because of the SectionOutlinerModel outliner items being used in the mixer which effectively allow you to use the mute/solo columns on sections.

#Synced-CL 51964473

diff --git a/Engine/Plugins/Experimental/MovieSceneAnimMixer/Source/MovieSceneAnimMixer/Private/MovieSceneAnimTransitionSectionBase.cpp b/Engine/Plugins/Experimental/MovieSceneAnimMixer/Source/MovieSceneAnimMixer/Private/MovieSceneAnimTransitionSectionBase.cpp
--- a/Engine/Plugins/Experimental/MovieSceneAnimMixer/Source/MovieSceneAnimMixer/Private/MovieSceneAnimTransitionSectionBase.cpp
+++ b/Engine/Plugins/Experimental/MovieSceneAnimMixer/Source/MovieSceneAnimMixer/Private/MovieSceneAnimTransitionSectionBase.cpp
@@ -235,12 +235,54 @@ void UMovieSceneAnimTransitionSectionBase::ImportEntityImpl(
 	}
 }
 
+static bool IsSectionEffectivelyDisabled(const UMovieSceneSection* Section)
+{
+	if (!Section || !Section->IsActive())
+	{
+		return true;
+	}
+
+#if WITH_EDITOR
+	if (Section->IsLocalEvalDisabled())
+	{
+		return true;
+	}
+#endif
+
+	if (const UMovieSceneTrack* Track = Section->GetTypedOuter<UMovieSceneTrack>())
+	{
+		if (Track->IsRowEvalDisabled(Section->GetRowIndex()))
+		{
+			return true;
+		}
+	}
+
+	return false;
+}
+
 bool UMovieSceneAnimTransitionSectionBase::PopulateEvaluationFieldImpl(
 	const TRange<FFrameNumber>& EffectiveRange,
 	const FMovieSceneEvaluationFieldEntityMetaData& InMetaData,
 	FMovieSceneEntityComponentFieldBuilder* OutFieldBuilder)
 {
+	// Skip evaluation if either the from or to section is disabled (muted, deactivated, or inactive)
+	if (IsSectionEffectivelyDisabled(FromSection) || IsSectionEffectivelyDisabled(ToSection))
+	{
+		return true;
+	}
+
 	const int32 MetaDataIndex = OutFieldBuilder->AddMetaData(InMetaData);
 	OutFieldBuilder->AddPersistentEntity(EffectiveRange, this, 0, MetaDataIndex);
 	return true;
 }
+
+#if WITH_EDITOR
+bool UMovieSceneAnimTransitionSectionBase::IsLocalEvalDisabled() const
+{
+	// A transition is disabled if it was directly muted or if either of its source sections
+	// is effectively disabled (muted, deactivated, row-disabled, or inactive).
+	return Super::IsLocalEvalDisabled()
+		|| IsSectionEffectivelyDisabled(FromSection)
+		|| IsSectionEffectivelyDisabled(ToSection);
+}
+#endif
diff --git a/Engine/Plugins/Experimental/MovieSceneAnimMixer/Source/MovieSceneAnimMixer/Public/MovieSceneAnimTransitionSectionBase.h b/Engine/Plugins/Experimental/MovieSceneAnimMixer/Source/MovieSceneAnimMixer/Public/MovieSceneAnimTransitionSectionBase.h
--- a/Engine/Plugins/Experimental/MovieSceneAnimMixer/Source/MovieSceneAnimMixer/Public/MovieSceneAnimTransitionSectionBase.h
+++ b/Engine/Plugins/Experimental/MovieSceneAnimMixer/Source/MovieSceneAnimMixer/Public/MovieSceneAnimTransitionSectionBase.h
@@ -98,6 +98,9 @@ public:
 
 	//~ UMovieSceneSection interface
 	virtual EMovieSceneChannelProxyType CacheChannelProxy() override final;
+#if WITH_EDITOR
+	MOVIESCENEANIMMIXER_API virtual bool IsLocalEvalDisabled() const override;
+#endif
 
 	//~ IMovieSceneEntityProvider interface
 	virtual void ImportEntityImpl(UMovieSceneEntitySystemLinker* EntityLinker, const FEntityImportParams& Params, FImportedEntity* OutImportedEntity) override;
diff --git a/Engine/Source/Editor/Sequencer/Private/SSequencer.cpp b/Engine/Source/Editor/Sequencer/Private/SSequencer.cpp
--- a/Engine/Source/Editor/Sequencer/Private/SSequencer.cpp
+++ b/Engine/Source/Editor/Sequencer/Private/SSequencer.cpp
@@ -1,6 +1,8 @@
 // Copyright Epic Games, Inc. All Rights Reserved.
 
 #include "SSequencer.h"
+#include "MVVM/Extensions/ISectionOwnerExtension.h"
+#include "MovieSceneSection.h"
 #include "BakingAnimationKeySettings.h"
 #include "Engine/Blueprint.h"
 #include "Filters/Filters/SequencerTrackFilter_CustomText.h"
@@ -2706,6 +2708,35 @@ void SSequencer::UpdateLayoutTree()
 				}
 			}
 
+			// Pass 3: Section owners that are not tracks or track rows (e.g. FSectionOutlinerModel).
+			// These own individual sections and need per-section eval-disable for mute/solo.
+			for (TViewModelPtr<ISectionOwnerExtension> SectionOwner : Sequencer->GetNodeTree()->GetRootNode()->GetDescendantsOfType<ISectionOwnerExtension>())
+			{
+				// Skip nodes already handled by the track or track row loops above
+				if (SectionOwner.AsModel()->IsA<ITrackExtension>() || SectionOwner.AsModel()->IsA<ITrackRowExtension>())
+				{
+					continue;
+				}
+
+				const ECachedMuteState MuteFlags = MuteState->GetCachedFlags(SectionOwner);
+				const ECachedSoloState SoloFlags = SoloState->GetCachedFlags(SectionOwner);
+
+				const bool bIsMuted = EnumHasAnyFlags(MuteFlags, ECachedMuteState::Muted | ECachedMuteState::ImplicitlyMutedByParent);
+				const bool bIsSoloed = EnumHasAnyFlags(SoloFlags, ECachedSoloState::Soloed | ECachedSoloState::ImplicitlySoloedByParent);
+
+				const bool bLocalEvalDisabled = bIsSoloed ? false : (bIsMuted || (bAnySoloNodes && !bIsSoloed));
+
+				for (UMovieSceneSection* Section : SectionOwner->GetSections())
+				{
+					if (Section && bLocalEvalDisabled != Section->IsLocalEvalDisabled())
+					{
+						Section->SetLocalEvalDisabled(bLocalEvalDisabled);
+						Section->MarkAsChanged();
+						bAnyChanged = true;
+					}
+				}
+			}
+
 			if (bAnyChanged)
 			{
 				Sequencer->NotifyMovieSceneDataChanged(EMovieSceneDataChangeType::TrackValueChanged);
diff --git a/Engine/Source/Editor/Sequencer/Private/SSequencerSection.cpp b/Engine/Source/Editor/Sequencer/Private/SSequencerSection.cpp
--- a/Engine/Source/Editor/Sequencer/Private/SSequencerSection.cpp
+++ b/Engine/Source/Editor/Sequencer/Private/SSequencerSection.cpp
@@ -1425,7 +1425,8 @@ int32 SSequencerSection::OnPaint( const FPaintArgs& Args, const FGeometry& Allot
 
 	const bool bConditionPasses = !Condition || MovieSceneHelpers::EvaluateSequenceCondition(BindingID, GetSequencer().GetFocusedTemplateID(), Condition, Track, GetSequencer().GetSharedPlaybackState());
 	const bool bTrackDisabled = Track && (Track->IsEvalDisabled() || Track->IsRowEvalDisabled(SectionObject->GetRowIndex()));
-	const bool bEnabled = bParentEnabled && SectionObject->IsActive() && bConditionPasses && !(bTrackDisabled);
+	const bool bSectionLocalDisabled = SectionObject->IsLocalEvalDisabled();
+	const bool bEnabled = bParentEnabled && SectionObject->IsActive() && bConditionPasses && !(bTrackDisabled) && !bSectionLocalDisabled;
 
 	const ESlateDrawEffect DrawEffects = bEnabled ? ESlateDrawEffect::None : ESlateDrawEffect::DisabledEffect;
 
diff --git a/Engine/Source/Runtime/MovieScene/Private/Compilation/MovieSceneCompiledDataManager.cpp b/Engine/Source/Runtime/MovieScene/Private/Compilation/MovieSceneCompiledDataManager.cpp
--- a/Engine/Source/Runtime/MovieScene/Private/Compilation/MovieSceneCompiledDataManager.cpp
+++ b/Engine/Source/Runtime/MovieScene/Private/Compilation/MovieSceneCompiledDataManager.cpp
@@ -1599,6 +1599,13 @@ void UMovieSceneCompiledDataManager::GatherTrack(const FMovieSceneBinding* Objec
 				continue;
 			}
 
+#if WITH_EDITOR
+			if (Entry.Section && Entry.Section->IsLocalEvalDisabled())
+			{
+				continue;
+			}
+#endif
+
 			IMovieSceneEntityProvider* EntityProvider = Cast<IMovieSceneEntityProvider>(Entry.Section);
 			if (!EntityProvider)
 			{
@@ -1650,6 +1657,13 @@ void UMovieSceneCompiledDataManager::GatherTrack(const FMovieSceneBinding* Objec
 					return true;
 				}
 
+#if WITH_EDITOR
+				if (Section && Section->IsLocalEvalDisabled())
+				{
+					return true;
+				}
+#endif
+
 				if (ClampedRangeRoot.IsEmpty())
 				{
 					return true;
diff --git a/Engine/Source/Runtime/MovieScene/Public/MovieSceneSection.h b/Engine/Source/Runtime/MovieScene/Public/MovieSceneSection.h
--- a/Engine/Source/Runtime/MovieScene/Public/MovieSceneSection.h
+++ b/Engine/Source/Runtime/MovieScene/Public/MovieSceneSection.h
@@ -602,6 +602,13 @@ public:
 	UFUNCTION(BlueprintPure, Category = "Sequencer|Section")
 	bool IsLocked() const { return bIsLocked; }
 
+#if WITH_EDITOR
+	/** Whether evaluation of this section is disabled due to local muting/soloing (not serialized).
+	 *  Virtual so subclasses (e.g. transitions) can add additional disable conditions. */
+	virtual bool IsLocalEvalDisabled() const { return bIsLocalEvalDisabled; }
+	void SetLocalEvalDisabled(bool bInEvalDisabled) { bIsLocalEvalDisabled = bInEvalDisabled; }
+#endif
+
 	/** Gets the number of frames to prepare this section for evaluation before it actually starts. */
 	UFUNCTION(BlueprintCallable, Category = "Sequencer|Section")
 	void SetPreRollFrames(int32 InPreRollFrames) { PreRollFrames = InPreRollFrames; }
@@ -831,6 +838,9 @@ private:
 	/** The color tint for this section */
 	UPROPERTY(EditAnywhere, Category = "Section")
 	FColor ColorTint;
+
+	/** Whether evaluation of this section is disabled due to local muting/soloing (transient, not serialized) */
+	bool bIsLocalEvalDisabled = false;
 #endif
 
 protected:
