Sequencer: Fix decoration category groups not showing keys when collapsed

Added FCategoryGroupModel::CreateTrackAreaCategory() to create and link
FCategoryModel track-area counterparts for decoration category groups.
Without this, GetTrackAreaModelList() was empty and no SChannelView was
created to render aggregated keys when a category group (e.g. Offset in
Root Motion Settings) was collapsed in the Sequencer outliner.

#Synced-CL 51964473

diff --git a/Engine/Source/Editor/Sequencer/Private/MVVM/ViewModels/CategoryModel.cpp b/Engine/Source/Editor/Sequencer/Private/MVVM/ViewModels/CategoryModel.cpp
--- a/Engine/Source/Editor/Sequencer/Private/MVVM/ViewModels/CategoryModel.cpp
+++ b/Engine/Source/Editor/Sequencer/Private/MVVM/ViewModels/CategoryModel.cpp
@@ -126,6 +126,14 @@ TArrayView<const TWeakViewModelPtr<FCategoryModel>> FCategoryGroupModel::GetCate
 	return Categories;
 }
 
+TSharedPtr<FCategoryModel> FCategoryGroupModel::CreateTrackAreaCategory()
+{
+	TSharedPtr<FCategoryModel> Category = MakeShared<FCategoryModel>(CategoryName);
+	AddCategory(Category);
+	Category->SetLinkedOutlinerItem(SharedThis(this));
+	return Category;
+}
+
 FOutlinerSizing FCategoryGroupModel::RecomputeSizing()
 {
 	FOutlinerSizing MaxSizing;
diff --git a/Engine/Source/Editor/Sequencer/Private/MVVM/ViewModels/DecorationOutlinerModel.cpp b/Engine/Source/Editor/Sequencer/Private/MVVM/ViewModels/DecorationOutlinerModel.cpp
--- a/Engine/Source/Editor/Sequencer/Private/MVVM/ViewModels/DecorationOutlinerModel.cpp
+++ b/Engine/Source/Editor/Sequencer/Private/MVVM/ViewModels/DecorationOutlinerModel.cpp
@@ -116,15 +116,15 @@ void FDecorationOutlinerModel::OnConstruct()
 	OutlinerChildren.MoveChildrenTo<IRecyclableExtension>(RecycledModels.GetChildren(), IRecyclableExtension::CallOnRecycle);
 	FViewModelChildren RecycledChildren = RecycledModels.GetChildren();
 
-	// Remove old channel models from their parent section models before destroying proxies.
-	for (const TWeakViewModelPtr<FChannelModel>& WeakChannel : TrackedChannelModels)
+	// Remove old channel and category models from their parent section models before destroying proxies.
+	for (const TWeakPtr<FViewModel>& WeakModel : TrackedSectionChildModels)
 	{
-		if (TViewModelPtr<FChannelModel> OldChannel = WeakChannel.Pin())
+		if (TSharedPtr<FViewModel> OldModel = WeakModel.Pin())
 		{
-			OldChannel->RemoveFromParent();
+			OldModel->RemoveFromParent();
 		}
 	}
-	TrackedChannelModels.Empty();
+	TrackedSectionChildModels.Empty();
 
 	// Clear old channel proxy before rebuilding
 	ChannelDecorationProxy.Reset();
@@ -226,6 +226,7 @@ void FDecorationOutlinerModel::OnConstruct()
 		// FCategoryGroupModel sub-group. Channels with no Group are added
 		// directly under this decoration model.
 		TMap<FName, TSharedPtr<FCategoryGroupModel>> CategoryGroupModels;
+		TMap<FName, TSharedPtr<FCategoryModel>> CategoryModels;
 		TMap<FName, TSharedPtr<FViewModel>> CategoryGroupTails;
 
 		TSharedPtr<FViewModel> ChannelTail;
@@ -237,18 +238,26 @@ void FDecorationOutlinerModel::OnConstruct()
 			FName ChannelName = MetaData.Name;
 
 			const bool bIsSubGroup = !MetaData.Group.IsEmpty();
+			const FName GroupName = bIsSubGroup ? FName(*MetaData.Group.ToString()) : NAME_None;
 
 			// Find or create the category group model for sub-grouped channels
 			TSharedPtr<FCategoryGroupModel> CategoryGroup;
 			if (bIsSubGroup)
 			{
-				FName GroupName(*MetaData.Group.ToString());
 				TSharedPtr<FCategoryGroupModel>& ExistingCategoryGroup = CategoryGroupModels.FindOrAdd(GroupName);
+
 				if (!ExistingCategoryGroup)
 				{
 					ExistingCategoryGroup = MakeShared<FCategoryGroupModel>(GroupName, MetaData.Group, MetaData.GetGroupTooltipTextDelegate);
 					OutlinerChildren.InsertChild(ExistingCategoryGroup, ChannelTail);
 					ChannelTail = ExistingCategoryGroup;
+
+					// Create a track-area FCategoryModel so the category group can
+					// render aggregated keys when collapsed in the outliner.
+					TSharedPtr<FCategoryModel> CategoryModel = ExistingCategoryGroup->CreateTrackAreaCategory();
+					SectionModel->GetChildList(FTrackModel::GetTopLevelChannelType()).AddChild(CategoryModel);
+					CategoryModels.Add(GroupName, CategoryModel);
+					TrackedSectionChildModels.Add(CategoryModel);
 				}
 				CategoryGroup = ExistingCategoryGroup;
 			}
@@ -291,7 +300,6 @@ void FDecorationOutlinerModel::OnConstruct()
 
 			if (bIsSubGroup)
 			{
-				FName GroupName(*MetaData.Group.ToString());
 				TSharedPtr<FViewModel>& GroupTail = CategoryGroupTails.FindOrAdd(GroupName);
 				CategoryGroup->GetChildList(EViewModelListType::Outliner).InsertChild(ChannelGroupModel, GroupTail);
 				GroupTail = ChannelGroupModel;
@@ -307,10 +315,24 @@ void FDecorationOutlinerModel::OnConstruct()
 			ChannelModel->SetLinkedOutlinerItem(ChannelGroupModel.ToSharedRef());
 			ChannelModel->Initialize(SectionInterface, ChannelHandle);
 
-			// Add channel model as child of section model so FindAncestorOfType<FSectionModel>() works
-			SectionModel->GetChildList(FTrackModel::GetTopLevelChannelType()).AddChild(ChannelModel);
+			// Parent under the track-area category model if this channel belongs to
+			// a sub-group, otherwise add directly to the section model.
+			if (bIsSubGroup)
+			{
+				CategoryModels[GroupName]->GetChildList(EViewModelListType::Generic).AddChild(ChannelModel);
+			}
+			else
+			{
+				SectionModel->GetChildList(FTrackModel::GetTopLevelChannelType()).AddChild(ChannelModel);
+			}
 
-			TrackedChannelModels.Add(ChannelModel);
+			TrackedSectionChildModels.Add(ChannelModel);
+		}
+
+		// Recompute sizing for category groups now that their categories are populated
+		for (const auto& Pair : CategoryGroupModels)
+		{
+			Pair.Value->RecomputeSizing();
 		}
 	}
 }
diff --git a/Engine/Source/Editor/Sequencer/Public/MVVM/ViewModels/CategoryModel.h b/Engine/Source/Editor/Sequencer/Public/MVVM/ViewModels/CategoryModel.h
--- a/Engine/Source/Editor/Sequencer/Public/MVVM/ViewModels/CategoryModel.h
+++ b/Engine/Source/Editor/Sequencer/Public/MVVM/ViewModels/CategoryModel.h
@@ -113,6 +113,11 @@ public:
 
 	UE_API TArrayView<const TWeakViewModelPtr<FCategoryModel>> GetCategories() const;
 
+	/** Creates a new FCategoryModel for the track area and links it to this category group.
+	 *  The caller is responsible for parenting the returned model in the appropriate
+	 *  track-area child list (e.g. SectionModel's TopLevelChannel list). */
+	UE_API TSharedPtr<FCategoryModel> CreateTrackAreaCategory();
+
 public:
 
 	/*~ ICompoundOutlinerExtension */
diff --git a/Engine/Source/Editor/Sequencer/Public/MVVM/ViewModels/DecorationOutlinerModel.h b/Engine/Source/Editor/Sequencer/Public/MVVM/ViewModels/DecorationOutlinerModel.h
--- a/Engine/Source/Editor/Sequencer/Public/MVVM/ViewModels/DecorationOutlinerModel.h
+++ b/Engine/Source/Editor/Sequencer/Public/MVVM/ViewModels/DecorationOutlinerModel.h
@@ -103,8 +103,8 @@ private:
 	// Channel proxy for IMovieSceneChannelDecoration. Stored to keep channel handles valid.
 	TSharedPtr<FMovieSceneChannelProxy> ChannelDecorationProxy;
 
-	// Tracks channel models added to section models so they can be removed during cleanup.
-	TArray<TWeakViewModelPtr<FChannelModel>> TrackedChannelModels;
+	// Tracks models (channel and category) added to section models so they can be removed during cleanup.
+	TArray<TWeakPtr<FViewModel>> TrackedSectionChildModels;
 
 	// Handle for this decoration's on signature changed delegate.
 	FDelegateHandle SignatureChangedHandle;
