Sequencer: Add layer bar with overlay text to Animation Mixer track

Show a layer bar on the Animation Mixer track displaying the combined
time range of all inner sections with overlay text indicating the layer
count (e.g. "3 Layers"). The layer bar is visible in both collapsed and
expanded states.

Adds optional overlay text support to SSequencerLayerBar via STextBlock,
plumbs it through FLayerBarModel, and creates the layer bar in
FAnimationMixerTrackModel::OnConstruct following the same pattern as
FObjectBindingModel.

#Synced-CL 51964473

diff --git a/Engine/Plugins/Experimental/MovieSceneAnimMixer/Source/MovieSceneAnimMixerEditor/Private/MVVM/AnimationMixerTrackModel.cpp b/Engine/Plugins/Experimental/MovieSceneAnimMixer/Source/MovieSceneAnimMixerEditor/Private/MVVM/AnimationMixerTrackModel.cpp
--- a/Engine/Plugins/Experimental/MovieSceneAnimMixer/Source/MovieSceneAnimMixerEditor/Private/MVVM/AnimationMixerTrackModel.cpp
+++ b/Engine/Plugins/Experimental/MovieSceneAnimMixer/Source/MovieSceneAnimMixerEditor/Private/MVVM/AnimationMixerTrackModel.cpp
@@ -24,6 +24,8 @@
 #include "ISequencer.h"
 #include "ISequencerTrackEditor.h"
 #include "MVVM/ViewModels/SequencerEditorViewModel.h"
+#include "MVVM/ViewModels/LayerBarModel.h"
+#include "SequencerSettings.h"
 #include "MovieSceneSection.h"
 #include "Styling/CoreStyle.h"
 #include "Styling/SlateTypes.h"
@@ -38,9 +40,53 @@ namespace UE::Sequencer
 
 UE_SEQUENCER_DEFINE_CASTABLE(FAnimationMixerTrackModel);
 
+static EViewModelListType GetLayerBarListType()
+{
+	static EViewModelListType LayerBarListType = RegisterCustomModelListType();
+	return LayerBarListType;
+}
+
 FAnimationMixerTrackModel::FAnimationMixerTrackModel(UMovieSceneTrack* Track)
 	: FTrackModel(Track)
+	, LayerBarList(GetLayerBarListType())
+{
+	RegisterChildList(&LayerBarList);
+}
+
+void FAnimationMixerTrackModel::OnConstruct()
+{
+	FTrackModel::OnConstruct();
+
+	if (!LayerBar)
+	{
+		TSharedPtr<FSequencerEditorViewModel> EditorViewModel = GetEditor();
+		TSharedPtr<ISequencer> Sequencer = EditorViewModel ? EditorViewModel->GetSequencer() : nullptr;
+
+		if (Sequencer && Sequencer->GetSequencerSettings()->GetShowLayerBars())
+		{
+			LayerBar = MakeShared<FLayerBarModel>(AsShared());
+			LayerBar->SetLinkedOutlinerItem(SharedThis(this));
+
+			UMovieSceneAnimationMixerTrack* MixerTrack = Cast<UMovieSceneAnimationMixerTrack>(GetTrack());
+			TWeakObjectPtr<UMovieSceneAnimationMixerTrack> WeakMixerTrack(MixerTrack);
+			LayerBar->SetOverlayText(MakeAttributeLambda([WeakMixerTrack]() -> FText
+			{
+				if (UMovieSceneAnimationMixerTrack* Track = WeakMixerTrack.Get())
+				{
+					const int32 NumLayers = Track->GetLayers().Num();
+					return FText::Format(NSLOCTEXT("AnimationMixerTrackModel", "LayerCount", "{0} {0}|plural(one=Layer,other=Layers)"), NumLayers);
+				}
+				return FText::GetEmpty();
+			}));
+
+			GetChildrenForList(&LayerBarList).AddChild(LayerBar);
+		}
+	}
+}
+
+FViewModelVariantIterator FAnimationMixerTrackModel::GetTopLevelChildTrackAreaModels() const
 {
+	return &LayerBarList;
 }
 
 bool FAnimationMixerTrackModel::IsDeactivated() const
diff --git a/Engine/Plugins/Experimental/MovieSceneAnimMixer/Source/MovieSceneAnimMixerEditor/Public/MVVM/AnimationMixerTrackModel.h b/Engine/Plugins/Experimental/MovieSceneAnimMixer/Source/MovieSceneAnimMixerEditor/Public/MVVM/AnimationMixerTrackModel.h
--- a/Engine/Plugins/Experimental/MovieSceneAnimMixer/Source/MovieSceneAnimMixerEditor/Public/MVVM/AnimationMixerTrackModel.h
+++ b/Engine/Plugins/Experimental/MovieSceneAnimMixer/Source/MovieSceneAnimMixerEditor/Public/MVVM/AnimationMixerTrackModel.h
@@ -6,6 +6,7 @@
 #include "Misc/FrameNumber.h"
 #include "Math/Range.h"
 #include "MVVM/ViewModels/TrackModel.h"
+#include "MVVM/ViewModels/ViewModelHierarchy.h"
 #include "MVVM/ViewModelPtr.h"
 
 class UMovieSceneAnimationMixerTrack;
@@ -13,6 +14,8 @@ class UMovieSceneAnimationMixerTrack;
 namespace UE::Sequencer
 {
 
+class FLayerBarModel;
+
 /**
  * Drop preview data for dragging sections in Animation Mixer track
  */
@@ -54,6 +57,12 @@ public:
 
 	explicit FAnimationMixerTrackModel(UMovieSceneTrack* Track);
 
+	/*~ FViewModel interface */
+	virtual void OnConstruct() override;
+
+	/*~ ITrackAreaExtension */
+	virtual FViewModelVariantIterator GetTopLevelChildTrackAreaModels() const override;
+
 	/*~ FEvaluableOutlinerItemModel */
 	virtual bool IsDeactivated() const override;
 	virtual void SetIsDeactivated(bool bInIsDeactivated) override;
@@ -92,6 +101,10 @@ private:
 
 	// Cached drop preview (only valid during drag)
 	mutable TOptional<FAnimMixerDropPreview> CachedPreview;
+
+	// Layer bar shown when the track is collapsed
+	FViewModelListHead LayerBarList;
+	TSharedPtr<FLayerBarModel> LayerBar;
 };
 
 } // namespace UE::Sequencer
diff --git a/Engine/Source/Editor/Sequencer/Private/MVVM/ViewModels/LayerBarModel.cpp b/Engine/Source/Editor/Sequencer/Private/MVVM/ViewModels/LayerBarModel.cpp
--- a/Engine/Source/Editor/Sequencer/Private/MVVM/ViewModels/LayerBarModel.cpp
+++ b/Engine/Source/Editor/Sequencer/Private/MVVM/ViewModels/LayerBarModel.cpp
@@ -76,9 +76,20 @@ void FLayerBarModel::OnExtensionsDirtied()
 {
 }
 
+void FLayerBarModel::SetOverlayText(TAttribute<FText> InText)
+{
+	OverlayText = MoveTemp(InText);
+}
+
+const TAttribute<FText>& FLayerBarModel::GetOverlayText() const
+{
+	return OverlayText;
+}
+
 TSharedPtr<ITrackLaneWidget> FLayerBarModel::CreateTrackLaneView(const FCreateTrackLaneViewParams& InParams)
 {
-	return SNew(SSequencerLayerBar, InParams.OwningTrackLane->GetTrackAreaView(), InParams.Editor->CastThisSharedChecked<FSequencerEditorViewModel>(), SharedThis(this));
+	return SNew(SSequencerLayerBar, InParams.OwningTrackLane->GetTrackAreaView(), InParams.Editor->CastThisSharedChecked<FSequencerEditorViewModel>(), SharedThis(this))
+		.OverlayText(OverlayText);
 }
 
 FTrackLaneVirtualAlignment FLayerBarModel::ArrangeVirtualTrackLaneView() const
diff --git a/Engine/Source/Editor/Sequencer/Private/Widgets/SSequencerLayerBar.cpp b/Engine/Source/Editor/Sequencer/Private/Widgets/SSequencerLayerBar.cpp
--- a/Engine/Source/Editor/Sequencer/Private/Widgets/SSequencerLayerBar.cpp
+++ b/Engine/Source/Editor/Sequencer/Private/Widgets/SSequencerLayerBar.cpp
@@ -22,6 +22,9 @@
 #include "ISequencerSection.h"
 
 #include "Widgets/Layout/SSpacer.h"
+#include "Widgets/Layout/SBox.h"
+#include "Widgets/SOverlay.h"
+#include "Widgets/Text/STextBlock.h"
 #include "Styling/AppStyle.h"
 
 #define LOCTEXT_NAMESPACE "SSequencerLayerBar"
@@ -221,40 +224,61 @@ void SSequencerLayerBar::Construct(const FArguments& InArgs, TWeakPtr<STrackArea
 	WeakLayerBar   = InWeakLayerBar;
 	WeakEditor     = InWeakEditor;
 	WeakTrackArea  = InWeakTrackArea;
+	OverlayText    = InArgs._OverlayText;
 
 	TSharedPtr<STrackAreaView> TrackArea = WeakTrackArea.Pin();
 
 	ChildSlot
 	[
-		SNew(SHorizontalBox)
-		.Visibility(this, &SSequencerLayerBar::GetHandleVisibility)
+		SNew(SOverlay)
 
-		+ SHorizontalBox::Slot()
-		.Padding(FMargin(1.f, 2.f))
-		.AutoWidth()
+		// Overlay text (always present, visible when text is non-empty)
+		+ SOverlay::Slot()
+		.HAlign(HAlign_Left)
+		.VAlign(VAlign_Center)
+		.Padding(6.f, 0.f)
 		[
-			SNew(SBox)
-			.WidthOverride(8.f)
-			[
-				SNew(SSequencerLayerBarHandle, TrackArea->GetViewModel(), SharedThis(this), EStretchConstraint::AnchorToEnd)
-				.BackgroundBrush(FAppStyle::Get().GetBrush("Sequencer.LayerBar.HandleLeft"))
-			]
+			SNew(STextBlock)
+			.Text(OverlayText)
+			.Font(FAppStyle::GetFontStyle("NormalFont"))
+			.ColorAndOpacity(FLinearColor::White)
+			.ShadowOffset(FVector2D(1.f, 1.f))
+			.ShadowColorAndOpacity(FLinearColor(0.f, 0.f, 0.f, 0.5f))
 		]
 
-		+ SHorizontalBox::Slot()
+		// Stretch handles (visible on hover)
+		+ SOverlay::Slot()
 		[
-			SNew(SSpacer)
-		]
+			SNew(SHorizontalBox)
+			.Visibility(this, &SSequencerLayerBar::GetHandleVisibility)
 
-		+ SHorizontalBox::Slot()
-		.Padding(FMargin(1.f, 2.f))
-		.AutoWidth()
-		[
-			SNew(SBox)
-			.WidthOverride(8.f)
+			+ SHorizontalBox::Slot()
+			.Padding(FMargin(1.f, 2.f))
+			.AutoWidth()
+			[
+				SNew(SBox)
+				.WidthOverride(8.f)
+				[
+					SNew(SSequencerLayerBarHandle, TrackArea->GetViewModel(), SharedThis(this), EStretchConstraint::AnchorToEnd)
+					.BackgroundBrush(FAppStyle::Get().GetBrush("Sequencer.LayerBar.HandleLeft"))
+				]
+			]
+
+			+ SHorizontalBox::Slot()
+			[
+				SNew(SSpacer)
+			]
+
+			+ SHorizontalBox::Slot()
+			.Padding(FMargin(1.f, 2.f))
+			.AutoWidth()
 			[
-				SNew(SSequencerLayerBarHandle, TrackArea->GetViewModel(), SharedThis(this), EStretchConstraint::AnchorToStart)
-				.BackgroundBrush(FAppStyle::Get().GetBrush("Sequencer.LayerBar.HandleRight"))
+				SNew(SBox)
+				.WidthOverride(8.f)
+				[
+					SNew(SSequencerLayerBarHandle, TrackArea->GetViewModel(), SharedThis(this), EStretchConstraint::AnchorToStart)
+					.BackgroundBrush(FAppStyle::Get().GetBrush("Sequencer.LayerBar.HandleRight"))
+				]
 			]
 		]
 	];
diff --git a/Engine/Source/Editor/Sequencer/Private/Widgets/SSequencerLayerBar.h b/Engine/Source/Editor/Sequencer/Private/Widgets/SSequencerLayerBar.h
--- a/Engine/Source/Editor/Sequencer/Private/Widgets/SSequencerLayerBar.h
+++ b/Engine/Source/Editor/Sequencer/Private/Widgets/SSequencerLayerBar.h
@@ -34,6 +34,7 @@ class SSequencerLayerBar
 {
 public:
 	SLATE_BEGIN_ARGS(SSequencerLayerBar) {}
+		SLATE_ATTRIBUTE(FText, OverlayText)
 	SLATE_END_ARGS()
 
 	void Construct(const FArguments& InArgs, TWeakPtr<STrackAreaView> InWeakTrackArea, TWeakPtr<FSequencerEditorViewModel> InWeakEditor, TWeakPtr<FLayerBarModel> InWeakLayerBar);
@@ -63,6 +64,7 @@ private:
 	TWeakPtr<FLayerBarModel> WeakLayerBar;
 	TWeakPtr<FSequencerEditorViewModel> WeakEditor;
 	TWeakPtr<STrackAreaView> WeakTrackArea;
+	TAttribute<FText> OverlayText;
 };
 
 } // namespace Sequencer
diff --git a/Engine/Source/Editor/Sequencer/Public/MVVM/ViewModels/LayerBarModel.h b/Engine/Source/Editor/Sequencer/Public/MVVM/ViewModels/LayerBarModel.h
--- a/Engine/Source/Editor/Sequencer/Public/MVVM/ViewModels/LayerBarModel.h
+++ b/Engine/Source/Editor/Sequencer/Public/MVVM/ViewModels/LayerBarModel.h
@@ -14,6 +14,8 @@
 #include "MVVM/ViewModels/ViewModel.h"
 #include "Math/Range.h"
 #include "Templates/SharedPointer.h"
+#include "Misc/Attribute.h"
+#include "Internationalization/Text.h"
 
 #define UE_API SEQUENCER_API
 
@@ -76,6 +78,9 @@ public:
 
 	UE_API void Offset(FFrameNumber Offset);
 
+	UE_API void SetOverlayText(TAttribute<FText> InText);
+	UE_API const TAttribute<FText>& GetOverlayText() const;
+
 private:
 
 	/*~ FViewModel interface */
@@ -84,6 +89,8 @@ private:
 
 	/*~ TViewModelExtensionCollection Interface */
 	UE_API void OnExtensionsDirtied() override;
+
+	TAttribute<FText> OverlayText;
 };
 
 } // namespace Sequencer
