Sequencer: Remove redundant word "Show" from some menus since they are under Visibility menus already

Show Marked Frames
Show Marked Frames Globally
Show Transport Controls
Show Bars
Show Value Indicators

#rnx
#Synced-CL 51926426

diff --git a/Engine/Source/Editor/CurveEditor/Private/CurveEditorCommands.cpp b/Engine/Source/Editor/CurveEditor/Private/CurveEditorCommands.cpp
--- a/Engine/Source/Editor/CurveEditor/Private/CurveEditorCommands.cpp
+++ b/Engine/Source/Editor/CurveEditor/Private/CurveEditorCommands.cpp
@@ -70,9 +70,9 @@ void FCurveEditorCommands::RegisterCommands()
 	UI_COMMAND(ToggleSnapTimeToSelection, "Snap Time to Selection", "Snap the current time to the first selected key time.", EUserInterfaceActionType::ToggleButton, FInputChord() );
 	UI_COMMAND(ToggleShowBufferedCurves, "Buffered Curves", "Show buffered curves for the selected curves.", EUserInterfaceActionType::ToggleButton, FInputChord() );
 	UI_COMMAND(ToggleShowCurveEditorCurveToolTips, "Curve Tool Tips", "Show a tool tip with name and values when hovering over a curve.", EUserInterfaceActionType::ToggleButton, FInputChord() );
-	UI_COMMAND(ToggleShowTransportControls, "Show Transport Controls", "Show transport controls (only used if the transport controls are not popped out)", EUserInterfaceActionType::ToggleButton, FInputChord());
-	UI_COMMAND(ToggleShowBars, "Show Bars", "Show Bars like Constraints and Spaces", EUserInterfaceActionType::ToggleButton, FInputChord());
-	UI_COMMAND(ToggleShowValueIndicatorLines, "Show Value Indicators", "Show horizontal, dotted lines on the selected min & max valued key handles", EUserInterfaceActionType::ToggleButton, FInputChord());
+	UI_COMMAND(ToggleShowTransportControls, "Transport Controls", "Show transport controls (only used if the transport controls are not popped out)", EUserInterfaceActionType::ToggleButton, FInputChord());
+	UI_COMMAND(ToggleShowBars, "Bars", "Show Bars like Constraints and Spaces", EUserInterfaceActionType::ToggleButton, FInputChord());
+	UI_COMMAND(ToggleShowValueIndicatorLines, "Value Indicators", "Show horizontal, dotted lines on the selected min & max valued key handles", EUserInterfaceActionType::ToggleButton, FInputChord());
 
 	UI_COMMAND(AddKeyHovered, "Add Key", "Add a new key to this curve at the current position.", EUserInterfaceActionType::Button, FInputChord() );
 	UI_COMMAND(PasteKeysHovered, "Paste", "Pastes the keys from the clipboard overwriting any key in destination track between the first and last pasted keys (default)", EUserInterfaceActionType::Button, FInputChord(EModifierKey::Control, EKeys::V));
diff --git a/Engine/Source/Editor/Sequencer/Private/SequencerCommands.cpp b/Engine/Source/Editor/Sequencer/Private/SequencerCommands.cpp
--- a/Engine/Source/Editor/Sequencer/Private/SequencerCommands.cpp
+++ b/Engine/Source/Editor/Sequencer/Private/SequencerCommands.cpp
@@ -124,12 +124,12 @@ void FSequencerCommands::RegisterCommands()
 	UI_COMMAND( SetKeyGroup, "Key Group", "Key the groups channels/properties when only one of them changes. ie. Keys all three translation channels when only translation Y changes", EUserInterfaceActionType::RadioButton, FInputChord());
 	UI_COMMAND( SetKeyAll, "Key All", "Key all channels/properties when only one of them changes. ie. Keys all translation, rotation, scale channels when only translation Y changes", EUserInterfaceActionType::RadioButton, FInputChord());
 
-	UI_COMMAND( ToggleShowMarkedFrames, "Show Marked Frames", "Show marked frames.", EUserInterfaceActionType::ToggleButton, FInputChord());
+	UI_COMMAND( ToggleShowMarkedFrames, "Marked Frames", "Show marked frames.", EUserInterfaceActionType::ToggleButton, FInputChord());
 	UI_COMMAND( ToggleMarkAtPlayPosition, "Toggle Mark", "Sets or clears a mark at the current play position.", EUserInterfaceActionType::Button, FInputChord(EKeys::M) );
 	UI_COMMAND( StepToNextMark, "Step to Next Marked Frame", "Step to the next marked frame", EUserInterfaceActionType::Button, FInputChord(EModifierKey::Control | EModifierKey::Shift, EKeys::Period) );
 	UI_COMMAND( StepToPreviousMark, "Step to Previous Marked Frame", "Step to the previous marked frame", EUserInterfaceActionType::Button, FInputChord(EModifierKey::Control | EModifierKey::Shift, EKeys::Comma) );
 	UI_COMMAND( ToggleMarksLocked, "Marked Frames Locked", "Prevent editing the marked frames.", EUserInterfaceActionType::ToggleButton, FInputChord() );
-	UI_COMMAND( ToggleShowMarkedFramesGlobally, "Show Marked Frames Globally", "Makes marked frames in this sub-sequence visible in parent/sibling sequences for the current editor session.", EUserInterfaceActionType::ToggleButton, FInputChord() );
+	UI_COMMAND( ToggleShowMarkedFramesGlobally, "Marked Frames Globally", "Makes marked frames in this sub-sequence visible in parent/sibling sequences for the current editor session.", EUserInterfaceActionType::ToggleButton, FInputChord() );
 	UI_COMMAND( ClearGlobalMarkedFrames, "Clear Global Marked Frames", "Set all marked frames in all sub-sequences to not be globally displayed.", EUserInterfaceActionType::Button, FInputChord() );
 
 	UI_COMMAND( ToggleAutoScroll, "Auto Scroll", "Toggle auto-scroll: When enabled, automatically scrolls the sequencer view to keep the current time visible", EUserInterfaceActionType::ToggleButton, FInputChord(EModifierKey::Shift, EKeys::S) );
