Sequencer: Hide/isolate selected tracks should not get all the tracks if there are none selected. Otherwise, all the tracks get hidden.

#Synced-CL 51996464

diff --git a/Engine/Source/Editor/Sequencer/Private/Filters/SequencerFilterBar.cpp b/Engine/Source/Editor/Sequencer/Private/Filters/SequencerFilterBar.cpp
--- a/Engine/Source/Editor/Sequencer/Private/Filters/SequencerFilterBar.cpp
+++ b/Engine/Source/Editor/Sequencer/Private/Filters/SequencerFilterBar.cpp
@@ -1594,22 +1594,28 @@ TSet<TWeakViewModelPtr<IOutlinerExtension>> FSequencerFilterBar::GetSelectedTrac
 	return SequencerViewModel ? ::GetSelectedTracksOrAll(*SequencerViewModel) : TSet<TWeakViewModelPtr<IOutlinerExtension>>();
 }
 
+TSet<TWeakViewModelPtr<IOutlinerExtension>> FSequencerFilterBar::GetSelectedTracks() const
+{
+	TSharedPtr<FSequencerEditorViewModel> SequencerViewModel = GetSequencer().GetViewModel();
+	return SequencerViewModel ? ::GetSelectedTracks(*SequencerViewModel) : TSet<TWeakViewModelPtr<IOutlinerExtension>>();
+}
+
 bool FSequencerFilterBar::HasSelectedTracks() const
 {
-	return !GetSelectedTracksOrAll().IsEmpty();
+	return !GetSelectedTracks().IsEmpty();
 }
 
 void FSequencerFilterBar::HideSelectedTracks()
 {
 	const bool bAddToExisting = !FSlateApplication::Get().GetModifierKeys().AreModifersDown(EModifierKey::Shift);
-	const TSet<TWeakViewModelPtr<IOutlinerExtension>> TracksToHide = GetSelectedTracksOrAll();
+	const TSet<TWeakViewModelPtr<IOutlinerExtension>> TracksToHide = GetSelectedTracks();
 	HideIsolateFilter->HideTracks(TracksToHide, bAddToExisting);
 }
 
 void FSequencerFilterBar::IsolateSelectedTracks()
 {
 	const bool bAddToExisting = FSlateApplication::Get().GetModifierKeys().AreModifersDown(EModifierKey::Shift);
-	const TSet<TWeakViewModelPtr<IOutlinerExtension>> TracksToIsolate = GetSelectedTracksOrAll();
+	const TSet<TWeakViewModelPtr<IOutlinerExtension>> TracksToIsolate = GetSelectedTracks();
 	HideIsolateFilter->IsolateTracks(TracksToIsolate, bAddToExisting);
 }
 
diff --git a/Engine/Source/Editor/Sequencer/Private/Filters/SequencerFilterBar.h b/Engine/Source/Editor/Sequencer/Private/Filters/SequencerFilterBar.h
--- a/Engine/Source/Editor/Sequencer/Private/Filters/SequencerFilterBar.h
+++ b/Engine/Source/Editor/Sequencer/Private/Filters/SequencerFilterBar.h
@@ -314,7 +314,10 @@ private:
 	
 	/** @return Gets the selected IOutlinerExtensions, or all of them if nothing is selected. */
 	TSet<UE::Sequencer::TWeakViewModelPtr<UE::Sequencer::IOutlinerExtension>> GetSelectedTracksOrAll() const;
-	
+
+	/** @return Gets the selected IOutlinerExtensions. */
+	TSet<UE::Sequencer::TWeakViewModelPtr<UE::Sequencer::IOutlinerExtension>> GetSelectedTracks() const;
+
 	/** Binds all commands that the filter bar or its filters handle. */
 	void BindCommands();
 	
diff --git a/Engine/Source/Editor/Sequencer/Private/Filters/Utils/HideIsolateViewModelUtils.cpp b/Engine/Source/Editor/Sequencer/Private/Filters/Utils/HideIsolateViewModelUtils.cpp
--- a/Engine/Source/Editor/Sequencer/Private/Filters/Utils/HideIsolateViewModelUtils.cpp
+++ b/Engine/Source/Editor/Sequencer/Private/Filters/Utils/HideIsolateViewModelUtils.cpp
@@ -24,6 +24,18 @@ TSet<TWeakViewModelPtr<IOutlinerExtension>> GetSelectedTracksOrAll(const FSequen
 	return SelectedSet.IsEmpty() ? GetAllTracks(InViewModel) : SelectedSet;
 }
 
+TSet<TWeakViewModelPtr<IOutlinerExtension>> GetSelectedTracks(const FSequencerEditorViewModel& InViewModel)
+{
+	const TSharedPtr<FSequencerSelection> Selection = InViewModel.GetSelection();
+	if (!Selection.IsValid())
+	{
+		return {};
+	}
+
+	const TSet<TWeakViewModelPtr<IOutlinerExtension>> SelectedSet = Selection->Outliner.GetSelected();
+	return SelectedSet;
+}
+
 TSet<TWeakViewModelPtr<IOutlinerExtension>> GetAllTracks(const FSequencerEditorViewModel& InViewModel)
 {
 	TSet<TWeakViewModelPtr<IOutlinerExtension>> OutTracks;
diff --git a/Engine/Source/Editor/Sequencer/Private/Filters/Utils/HideIsolateViewModelUtils.h b/Engine/Source/Editor/Sequencer/Private/Filters/Utils/HideIsolateViewModelUtils.h
--- a/Engine/Source/Editor/Sequencer/Private/Filters/Utils/HideIsolateViewModelUtils.h
+++ b/Engine/Source/Editor/Sequencer/Private/Filters/Utils/HideIsolateViewModelUtils.h
@@ -17,6 +17,8 @@ class IOutlinerExtension;
 
 /** @return Gets the selected IOutlinerExtensions, or all of them if nothing is selected. */
 TSet<TWeakViewModelPtr<IOutlinerExtension>> GetSelectedTracksOrAll(const FSequencerEditorViewModel& InViewModel);
+/** @return Gets the selected IOutlinerExtensions. */
+TSet<TWeakViewModelPtr<IOutlinerExtension>> GetSelectedTracks(const FSequencerEditorViewModel& InViewModel);
 /** @return Gets all tracks */
 TSet<TWeakViewModelPtr<IOutlinerExtension>> GetAllTracks(const FSequencerEditorViewModel& InViewModel);
 
diff --git a/Engine/Source/Editor/Sequencer/Private/Filters/ViewModels/HideIsolateShowViewModel.cpp b/Engine/Source/Editor/Sequencer/Private/Filters/ViewModels/HideIsolateShowViewModel.cpp
--- a/Engine/Source/Editor/Sequencer/Private/Filters/ViewModels/HideIsolateShowViewModel.cpp
+++ b/Engine/Source/Editor/Sequencer/Private/Filters/ViewModels/HideIsolateShowViewModel.cpp
@@ -158,6 +158,13 @@ TSet<TWeakViewModelPtr<IOutlinerExtension>> FHideIsolateShowViewModel::GetSelect
 		? AllTracksAttr.Get()
 		: Selection;
 }
+
+TSet<TWeakViewModelPtr<IOutlinerExtension>> FHideIsolateShowViewModel::GetSelectedTracks() const
+{
+	const TSet<TWeakViewModelPtr<IOutlinerExtension>> Selection = SelectedTracksAttr.Get();
+	return Selection;
+}
+
 } // namespace UE::Sequencer
 
 #undef LOCTEXT_NAMESPACE
\ No newline at end of file
diff --git a/Engine/Source/Editor/Sequencer/Private/Filters/ViewModels/HideIsolateShowViewModel.h b/Engine/Source/Editor/Sequencer/Private/Filters/ViewModels/HideIsolateShowViewModel.h
--- a/Engine/Source/Editor/Sequencer/Private/Filters/ViewModels/HideIsolateShowViewModel.h
+++ b/Engine/Source/Editor/Sequencer/Private/Filters/ViewModels/HideIsolateShowViewModel.h
@@ -80,6 +80,9 @@ private:
 	
 	/** @return Selected tracks, or all of them. */
 	TSet<TWeakViewModelPtr<IOutlinerExtension>> GetSelectedTracksOrAll() const;
+
+	/** @return Selected tracks. */
+	TSet<TWeakViewModelPtr<IOutlinerExtension>> GetSelectedTracks() const;
 };
 } // namespace UE::Sequencer
 
