Sequencer: Callback that's fired before we activate to another sequence, this way we can do stuff on that current sequence before jumping to the next one.

#Synced-CL 51987680

diff --git a/Engine/Source/Editor/Sequencer/Private/Sequencer.cpp b/Engine/Source/Editor/Sequencer/Private/Sequencer.cpp
--- a/Engine/Source/Editor/Sequencer/Private/Sequencer.cpp
+++ b/Engine/Source/Editor/Sequencer/Private/Sequencer.cpp
@@ -1347,6 +1347,8 @@ void FSequencer::ResetToNewRootSequence(UMovieSceneSequence& NewSequence)
 {
 	using namespace UE::MovieScene;
 
+	OnPreNavigateToSequenceEvent.Broadcast();
+
 	RemoveNodeGroupsCollectionChangedDelegate();
 
 	const UMovieSceneSequence* PreviousRootSequence = RootSequence.Get();
@@ -1402,6 +1404,8 @@ void FSequencer::ResetToNewRootSequence(UMovieSceneSequence& NewSequence)
 
 void FSequencer::FocusSequenceInstance(UMovieSceneSubSection& InSubSection)
 {
+	OnPreNavigateToSequenceEvent.Broadcast();
+
 	RemoveNodeGroupsCollectionChangedDelegate();
 
 	const UMovieSceneSequence* PreviousFocusedSequence = GetFocusedMovieSceneSequence();
@@ -1636,6 +1640,8 @@ void FSequencer::PopToSequenceInstance(FMovieSceneSequenceIDRef SequenceID)
 {
 	if( ActiveTemplateIDs.Num() > 1 )
 	{
+		OnPreNavigateToSequenceEvent.Broadcast();
+
 		TemplateIDBackwardStack.Push(ActiveTemplateIDs.Top());
 		TemplateIDForwardStack.Reset();
 
@@ -4253,6 +4259,11 @@ ISequencer::FOnPostSave& FSequencer::OnPostSave()
 	return OnPostSaveEvent;
 }
 
+ISequencer::FOnPreNavigateToSequence& FSequencer::OnPreNavigateToSequence()
+{
+	return OnPreNavigateToSequenceEvent;
+}
+
 ISequencer::FOnActivateSequence& FSequencer::OnActivateSequence()
 {
 	return OnActivateSequenceEvent;
diff --git a/Engine/Source/Editor/Sequencer/Private/Sequencer.h b/Engine/Source/Editor/Sequencer/Private/Sequencer.h
--- a/Engine/Source/Editor/Sequencer/Private/Sequencer.h
+++ b/Engine/Source/Editor/Sequencer/Private/Sequencer.h
@@ -840,6 +840,7 @@ public:
 	virtual FOnActorAddedToSequencer& OnActorAddedToSequencer() override;
 	virtual FOnPreSave& OnPreSave() override;
 	virtual FOnPostSave& OnPostSave() override;
+	virtual FOnPreNavigateToSequence& OnPreNavigateToSequence() override;
 	virtual FOnActivateSequence& OnActivateSequence() override;
 	virtual FOnCameraCut& OnCameraCut() override;
 	virtual void AddNumericTypeInterface(TSharedRef<UE::Sequencer::FSequencerNumericTypeInterface> InNumericTypeInterface);
@@ -1547,8 +1548,9 @@ private:
 	FOnCameraCut OnCameraCutEvent;
 	FOnPreSave OnPreSaveEvent;
 	FOnPostSave OnPostSaveEvent;
+	FOnPreNavigateToSequence OnPreNavigateToSequenceEvent;
 	FOnActivateSequence OnActivateSequenceEvent;
-
+	
 	/** Delegate for Curve Display Changed Event from the Curve Editor, which we than pass to the FOnCurveDisplayChanged delegate */
 	void OnCurveModelDisplayChanged(FCurveModel *InCurveModel, bool bDisplayed, const FCurveEditor* InCurveEditor);
 
diff --git a/Engine/Source/Editor/Sequencer/Public/ISequencer.h b/Engine/Source/Editor/Sequencer/Public/ISequencer.h
--- a/Engine/Source/Editor/Sequencer/Public/ISequencer.h
+++ b/Engine/Source/Editor/Sequencer/Public/ISequencer.h
@@ -600,6 +600,9 @@ public:
 	DECLARE_EVENT_OneParam(ISequencer, FOnPostSave, ISequencer&)
 	virtual FOnPostSave& OnPostSave() = 0;
 
+	DECLARE_EVENT(ISequencer, FOnPreNavigateToSequence)
+	virtual FOnPreNavigateToSequence& OnPreNavigateToSequence() = 0;
+
 	DECLARE_EVENT_OneParam(ISequencer, FOnActivateSequence, FMovieSceneSequenceIDRef)
 	virtual FOnActivateSequence& OnActivateSequence() = 0;
 
