Sequencer: Rename FLinkedFilterMediator to FOutlinerFilterAreaManager.

#Synced-CL 51881412

diff --git a/Engine/Source/Editor/Sequencer/Private/MVVM/ViewModels/LinkedFilteringMediator.cpp b/Engine/Source/Editor/Sequencer/Private/MVVM/ViewModels/LinkedFilteringMediator.cpp
deleted file mode 100644
diff --git a/Engine/Source/Editor/Sequencer/Private/MVVM/ViewModels/LinkedFilteringMediator.h b/Engine/Source/Editor/Sequencer/Private/MVVM/ViewModels/LinkedFilteringMediator.h
deleted file mode 100644
diff --git a/Engine/Source/Editor/Sequencer/Private/MVVM/ViewModels/OutlinerFilterAreaManager.cpp b/Engine/Source/Editor/Sequencer/Private/MVVM/ViewModels/OutlinerFilterAreaManager.cpp
new file mode 100644
--- /dev/null
+++ b/Engine/Source/Editor/Sequencer/Private/MVVM/ViewModels/OutlinerFilterAreaManager.cpp
@@ -0,0 +1,76 @@
+﻿// Copyright Epic Games, Inc. All Rights Reserved.
+
+#include "OutlinerFilterAreaManager.h"
+
+#include "Filters/SequencerFilterBar.h"
+#include "Filters/SequencerTrackFilterCommands.h"
+#include "Filters/Utils/LinkedFilteringViewUtils.h"
+#include "Filters/Linking/Mode/ILinkedFilterViewModel.h"
+#include "Filters/Linking/Mode/LinkedFilterViewModel.h"
+#include "SequencerNodeTree.h"
+#include "Filters/FilterConfigIdentifiers.h"
+
+namespace UE::Sequencer
+{
+FOutlinerFilterAreaManager::FOutlinerFilterAreaManager(
+	TWeakPtr<FSequencer> InSequencer, const TSharedRef<FLinkedFilterViewModel>& InSequencerFilterViewModel
+	)
+	: FFilterAreaManager(ConfigIds::FilterArea_Outliner, MoveTemp(InSequencer), InSequencerFilterViewModel)
+{
+	GetFilterModel()->OnFilterModeChanged().AddRaw(this, &FOutlinerFilterAreaManager::OnFilterModeChanged);
+	
+	const TSharedRef<FSequencerFilterBar> SharedFilterBar = GetFilterModel()->GetLinkedFilterBarImpl();
+	SharedFilterBar->OnChangeOutlinerExtensionFilterState().AddRaw(this, &FOutlinerFilterAreaManager::OnFilterStateChanged_LinkedFilterBar);
+	
+	const TSharedRef<FSequencerFilterBar> InstancedFilterBar = GetFilterModel()->GetInstancedFilterBarImpl();
+	InstancedFilterBar->OnChangeOutlinerExtensionFilterState().AddRaw(this, &FOutlinerFilterAreaManager::OnFilterStateChanged_InstancedFilterBar);
+}
+
+FOutlinerFilterAreaManager::~FOutlinerFilterAreaManager()
+{
+	GetFilterModel()->OnFilterModeChanged().RemoveAll(this);
+	
+	const TSharedRef<FSequencerFilterBar> SharedFilterBar = GetFilterModel()->GetLinkedFilterBarImpl();
+	SharedFilterBar->OnChangeOutlinerExtensionFilterState().RemoveAll(this);
+	
+	const TSharedRef<FSequencerFilterBar> InstancedFilterBar = GetFilterModel()->GetInstancedFilterBarImpl();
+	InstancedFilterBar->OnChangeOutlinerExtensionFilterState().RemoveAll(this);
+}
+
+void FOutlinerFilterAreaManager::BindCommands(const TSharedRef<FUICommandList>& InSequencerCommandList) const
+{
+	InSequencerCommandList->Append(GetFilterAreaCommandList());
+}
+
+ELinkedFilterMode FOutlinerFilterAreaManager::GetFilterMode() const
+{
+	return GetFilterModel()->GetFilterMode();
+}
+
+void FOutlinerFilterAreaManager::OnFilterModeChanged()
+{
+	// TODO: The goal of this is to call SetFilteredOut on all items again, but it does more work then it needs to. Instead of rerunning the filter,
+	// we could iterate through the hierarchy, and use FSequencerFilterBar::GetFilterData to call SetFilteredOut with the correct filter state.
+	GetFilterModel()->GetActiveFilterBar()->RequestFilterUpdate();
+}
+
+void FOutlinerFilterAreaManager::OnFilterStateChanged_LinkedFilterBar(
+	const TViewModelPtr<IOutlinerExtension>& InItem, bool bInIsFilteredOut
+	) const
+{
+	if (GetFilterMode() == ELinkedFilterMode::Linked)
+	{
+		InItem->SetFilteredOut(bInIsFilteredOut);
+	}
+}
+
+void FOutlinerFilterAreaManager::OnFilterStateChanged_InstancedFilterBar(
+	const TViewModelPtr<IOutlinerExtension>& InItem, bool bInIsFilteredOut
+	) const
+{
+	if (GetFilterMode() == ELinkedFilterMode::Instanced)
+	{
+		InItem->SetFilteredOut(bInIsFilteredOut);
+	}
+}
+} // namespace UE::Sequencer
diff --git a/Engine/Source/Editor/Sequencer/Private/MVVM/ViewModels/OutlinerFilterAreaManager.h b/Engine/Source/Editor/Sequencer/Private/MVVM/ViewModels/OutlinerFilterAreaManager.h
new file mode 100644
--- /dev/null
+++ b/Engine/Source/Editor/Sequencer/Private/MVVM/ViewModels/OutlinerFilterAreaManager.h
@@ -0,0 +1,48 @@
+﻿// Copyright Epic Games, Inc. All Rights Reserved.
+
+#pragma once
+
+#include "Filters/Linking/FilterAreaManager.h"
+#include "MVVM/Extensions/IOutlinerExtension.h"
+#include "MVVM/ViewModels/ViewModel.h"
+
+class FUICommandList;
+
+namespace UE::Sequencer
+{
+class FLinkedFilterViewModel;
+enum class ELinkedFilterMode : uint8;
+
+/** 
+ * Hooks up Sequencer's own FLinkedFilterViewModel with the relevant systems.
+ * 
+ * Its responsibility is to handle updating any systems when the linked filtering mode changes.
+ * For example:
+ * - Calling IOutlinerExtension::SetFilteredOut when the item's filter state changes
+ * - Binding commands that are bound by FSequencerFilterBar and making sure correct commands are executed depending on which one is active.
+ */
+class FOutlinerFilterAreaManager : public FFilterAreaManager
+{
+public:
+	
+	explicit FOutlinerFilterAreaManager(TWeakPtr<FSequencer> InSequencer, const TSharedRef<FLinkedFilterViewModel>& InSequencerFilterViewModel);
+	~FOutlinerFilterAreaManager();
+	
+	/** Binds the filter bar commands. */
+	void BindCommands(const TSharedRef<FUICommandList>& InSequencerCommandList) const;
+	
+private:
+	
+	/** @return The filter mode that Sequencer is currently using. */
+	ELinkedFilterMode GetFilterMode() const;
+
+	/** Schedules the IOutlinerExtensions to be updated with SetFilteredOut calls to reflect the state of the new filter bar. */
+	void OnFilterModeChanged();
+	
+	/** Handles propagating the IOutlinerExtension's filter out state if the active filter bar is the linked one. */
+	void OnFilterStateChanged_LinkedFilterBar(const TViewModelPtr<IOutlinerExtension>& InItem, bool bInIsFilteredOut) const;
+	
+	/** Handles propagating the IOutlinerExtension's filter out state if the active filter bar is the instanced one. */
+	void OnFilterStateChanged_InstancedFilterBar(const TViewModelPtr<IOutlinerExtension>& InItem, bool bInIsFilteredOut) const;
+};
+} // namespace UE::Sequencer
diff --git a/Engine/Source/Editor/Sequencer/Private/MVVM/ViewModels/SequencerEditorViewModelImpl.cpp b/Engine/Source/Editor/Sequencer/Private/MVVM/ViewModels/SequencerEditorViewModelImpl.cpp
--- a/Engine/Source/Editor/Sequencer/Private/MVVM/ViewModels/SequencerEditorViewModelImpl.cpp
+++ b/Engine/Source/Editor/Sequencer/Private/MVVM/ViewModels/SequencerEditorViewModelImpl.cpp
@@ -6,7 +6,7 @@
 #include "Filters/Linking/Mode/LinkedFilterFactoryViewModel.h"
 #include "Filters/Linking/Mode/LinkedFilterViewModel.h"
 #include "Filters/Linking/Mode/MakeLinkedFilterViewModelArgs.h"
-#include "LinkedFilteringMediator.h"
+#include "OutlinerFilterAreaManager.h"
 #include "Sequencer.h"
 
 namespace UE::Sequencer
@@ -47,6 +47,6 @@ void FSequencerEditorViewModelImpl::InitLinkedFiltering()
 	SequenceFilterInstanceViewModel = LinkedFilterFactory->MakeFilteringModelImpl(
 		FMakeLinkedFilterViewModelArgs(ConfigIds::Sequencer_InstancedConfigSubId, UnlinkedFilterFlags)
 		);
-	FilterMediator = MakeShared<FLinkedFilteringMediator>(Sequencer, SequenceFilterInstanceViewModel.ToSharedRef());
+	OutlinerFilterArea = MakeShared<FOutlinerFilterAreaManager>(Sequencer, SequenceFilterInstanceViewModel.ToSharedRef());
 }
 }
diff --git a/Engine/Source/Editor/Sequencer/Private/MVVM/ViewModels/SequencerEditorViewModelImpl.h b/Engine/Source/Editor/Sequencer/Private/MVVM/ViewModels/SequencerEditorViewModelImpl.h
--- a/Engine/Source/Editor/Sequencer/Private/MVVM/ViewModels/SequencerEditorViewModelImpl.h
+++ b/Engine/Source/Editor/Sequencer/Private/MVVM/ViewModels/SequencerEditorViewModelImpl.h
@@ -6,7 +6,7 @@
 
 namespace UE::Sequencer
 {
-class FLinkedFilteringMediator;
+class FOutlinerFilterAreaManager;
 class FLinkedFilterViewModel;
 class FLinkedFilterFactoryViewModel;
 
@@ -22,7 +22,7 @@ public:
 	TSharedPtr<FLinkedFilterFactoryViewModel> GetLinkedFilterViewModelFactoryImpl() const { return LinkedFilterFactory; }
 	
 	/** @return The object responsible for syncing the linked filters to Sequencer's systems. */
-	TSharedPtr<FLinkedFilteringMediator> GetFilterMediator() const { return FilterMediator; }
+	TSharedPtr<FOutlinerFilterAreaManager> GetOutlinerFilterArea() const { return OutlinerFilterArea; }
 	
 protected:
 	
@@ -38,7 +38,7 @@ private:
 	TSharedPtr<FLinkedFilterViewModel> SequenceFilterInstanceViewModel;
 	
 	/** Hooks up Sequencer's own FLinkedFilterViewModel with the relevant systems.*/
-	TSharedPtr<FLinkedFilteringMediator> FilterMediator;
+	TSharedPtr<FOutlinerFilterAreaManager> OutlinerFilterArea;
 	
 	/** Inits the systems relevant for linked filtering. */
 	void InitLinkedFiltering();
diff --git a/Engine/Source/Editor/Sequencer/Private/SSequencer.cpp b/Engine/Source/Editor/Sequencer/Private/SSequencer.cpp
--- a/Engine/Source/Editor/Sequencer/Private/SSequencer.cpp
+++ b/Engine/Source/Editor/Sequencer/Private/SSequencer.cpp
@@ -161,7 +161,7 @@
 #include "WorkspaceMenuStructureModule.h"
 #include "Filters/FilterConfigIdentifiers.h"
 #include "Filters/Widgets/Linking/SFilterBarSwitcher.h"
-#include "MVVM/ViewModels/LinkedFilteringMediator.h"
+#include "MVVM/ViewModels/OutlinerFilterAreaManager.h"
 #include "MVVM/ViewModels/SequencerEditorViewModelImpl.h"
 
 #define LOCTEXT_NAMESPACE "Sequencer"
diff --git a/Engine/Source/Editor/Sequencer/Private/Sequencer.cpp b/Engine/Source/Editor/Sequencer/Private/Sequencer.cpp
--- a/Engine/Source/Editor/Sequencer/Private/Sequencer.cpp
+++ b/Engine/Source/Editor/Sequencer/Private/Sequencer.cpp
@@ -193,7 +193,7 @@
 #include "Decorations/MovieSceneMuteSoloDecoration.h"
 #include "Filters/SequencerTrackFilterCommands.h"
 #include "Filters/Linking/Mode/LinkedFilterViewModel.h"
-#include "MVVM/ViewModels/LinkedFilteringMediator.h"
+#include "MVVM/ViewModels/OutlinerFilterAreaManager.h"
 #include "MVVM/ViewModels/SequencerEditorViewModelImpl.h"
 #include "TrackEditors/TimeWarpTrackEditor.h"
 #include "Variants/MovieSceneTimeWarpGetter.h"
@@ -12413,10 +12413,10 @@ void FSequencer::BindCommands()
 	// bind widget specific commands
 	SequencerWidget->BindCommands(SequencerCommandBindings);
 	
-	ViewModel->GetFilterMediator()->BindCommands(SequencerCommandBindings);
+	ViewModel->GetOutlinerFilterArea()->BindCommands(SequencerCommandBindings);
 	
 	// If this sequencer supports a curve editor, let's add bindings for it. 
-	// As FLinkedFilteringMediator::BindCommands binds the shared filter bar commands, which the curve editor also binds to, this must happen after.
+	// As FOutlinerFilterAreaManager::BindCommands binds the shared filter bar commands, which the curve editor also binds to, this must happen after.
 	if (FSequencerCurveEditorApp* CurveEditorIntegration = FSequencerCurveEditorApp::Get(*this))
 	{
 		// We want a subset of the commands to work in the Curve Editor too, but bound to our functions. This minimizes code duplication
