Sequencer: Fix setting the selection range start and end so that they respect the snapping values

#Synced-CL 52147848

diff --git a/Engine/Source/Editor/Sequencer/Private/Sequencer.cpp b/Engine/Source/Editor/Sequencer/Private/Sequencer.cpp
--- a/Engine/Source/Editor/Sequencer/Private/Sequencer.cpp
+++ b/Engine/Source/Editor/Sequencer/Private/Sequencer.cpp
@@ -3026,6 +3026,7 @@ void FSequencer::SetSelectionRangeEnd(FFrameTime EndFrame)
 	using namespace UE::MovieScene;
 
 	const bool bInitiallyEmpty = GetSelectionRange().IsEmpty();
+	SnapSequencerTime(EndFrame);
 	const FFrameNumber LocalTime = EndFrame.FrameNumber;
 
 	const FFrameNumber StartFrame = (bInitiallyEmpty || GetSelectionRange().GetLowerBoundValue() >= LocalTime)
@@ -3053,6 +3054,7 @@ void FSequencer::SetSelectionRangeStart(FFrameTime StartFrame)
 	using namespace UE::MovieScene;
 
 	const bool bInitiallyEmpty = GetSelectionRange().IsEmpty();
+	SnapSequencerTime(StartFrame);
 	const FFrameNumber LocalTime = StartFrame.FrameNumber;
 
 	const FFrameNumber EndFrame = (bInitiallyEmpty || GetSelectionRange().GetUpperBoundValue() <= LocalTime)
