Sequencer: Hide the linked filtering toggle button when there is no other system using the system.

This is implemented by checking the count of valid FLinkedFilteringViewModels. If the count is 1, then only Sequencer is instantiated. If the count is more, than another system is assumed to be using it. For now, Curve Editor is the only other system using it, so count should be either 1 or 2. This implementation was done so we don't need to modify this code if another system other than Curve Editor is added in the future.

#Synced-CL 51881412

diff --git a/Engine/Source/Editor/Sequencer/Private/Filters/Linking/Mode/LinkedFilterFactoryViewModel.cpp b/Engine/Source/Editor/Sequencer/Private/Filters/Linking/Mode/LinkedFilterFactoryViewModel.cpp
--- a/Engine/Source/Editor/Sequencer/Private/Filters/Linking/Mode/LinkedFilterFactoryViewModel.cpp
+++ b/Engine/Source/Editor/Sequencer/Private/Filters/Linking/Mode/LinkedFilterFactoryViewModel.cpp
@@ -2,8 +2,9 @@
 
 #include "LinkedFilterFactoryViewModel.h"
 
-#include "LinkedFilterViewModel.h"
+#include "Algo/Count.h"
 #include "Filters/Linking/Mode/MakeLinkedFilterViewModelArgs.h"
+#include "LinkedFilterViewModel.h"
 
 namespace UE::Sequencer
 {
@@ -17,12 +18,22 @@ FLinkedFilterFactoryViewModel::FLinkedFilterFactoryViewModel(
 	, FilterEvaluator(InFilterEvaluator)
 {}
 
-TSharedRef<FLinkedFilterViewModel> FLinkedFilterFactoryViewModel::MakeFilteringModelImpl(const FMakeLinkedFilterViewModelArgs& InArgs) const
+int32 FLinkedFilterFactoryViewModel::GetFilterFeatureUseCount() const
+{
+	return Algo::CountIf(FeatureUsers, [](const TWeakPtr<FLinkedFilterViewModel>& ViewModel)
+	{
+		return ViewModel.IsValid();
+	});
+}
+
+TSharedRef<FLinkedFilterViewModel> FLinkedFilterFactoryViewModel::MakeFilteringModelImpl(const FMakeLinkedFilterViewModelArgs& InArgs)
 {
 	const TSharedRef<FSequencerFilterBar> InstancedFilterBar = FSequencerFilterBar::Make(
 		FFilterBarInitArgs(Sequencer, InArgs.ConfigSubIdentifier, InArgs.UnlinkedFilterFlags)
 		);
-	return MakeShared<FLinkedFilterViewModel>(SharedFilterBar, InstancedFilterBar, FilterEvaluator);
+	const TSharedRef<FLinkedFilterViewModel> ViewModel = MakeShared<FLinkedFilterViewModel>(SharedFilterBar, InstancedFilterBar, FilterEvaluator);
+	FeatureUsers.Emplace(ViewModel.ToWeakPtr());
+	return ViewModel;
 }
 
 TSharedRef<ILinkedFilterViewModel> FLinkedFilterFactoryViewModel::MakeFilteringModel(const FMakeLinkedFilterViewModelArgs& InArgs)
diff --git a/Engine/Source/Editor/Sequencer/Private/Filters/Linking/Mode/LinkedFilterFactoryViewModel.h b/Engine/Source/Editor/Sequencer/Private/Filters/Linking/Mode/LinkedFilterFactoryViewModel.h
--- a/Engine/Source/Editor/Sequencer/Private/Filters/Linking/Mode/LinkedFilterFactoryViewModel.h
+++ b/Engine/Source/Editor/Sequencer/Private/Filters/Linking/Mode/LinkedFilterFactoryViewModel.h
@@ -11,6 +11,10 @@ namespace UE::Sequencer
 class FFilterEvaluator;
 class FLinkedFilterViewModel;
 
+/**
+ * Knows how to create FLinkedFilterViewModel. 
+ * This factory keep a shared FSequencerFilterBar state that is used when a FLinkedFilterViewModel is put into linked mode.
+ */
 class FLinkedFilterFactoryViewModel : public ILinkedFilterFactoryViewModel
 {
 public:
@@ -20,9 +24,17 @@ public:
 		const TSharedRef<FSequencerFilterBar>& InSharedFilterBar,
 		const TSharedRef<FFilterEvaluator> InFilterEvaluator
 		);
+
+	/**
+	 * @return The number of features using the linked filtering feature.
+	 * 
+	 * @note This is determined by looking at the number of valid TWeakPtr<FLinkedFilterViewModel>. 
+	 * It is assumed that if a pointer is valid, the system that created it is still using it.
+	 */
+	int32 GetFilterFeatureUseCount() const;
 	
 	/** @return View model that manages linked and instanced filter state for your UI.  */
-	TSharedRef<FLinkedFilterViewModel> MakeFilteringModelImpl(const FMakeLinkedFilterViewModelArgs& InArgs) const;
+	TSharedRef<FLinkedFilterViewModel> MakeFilteringModelImpl(const FMakeLinkedFilterViewModelArgs& InArgs);
 	
 	//~ Begin ILinkedFilterFactoryViewModel
 	virtual TSharedRef<ILinkedFilterViewModel> MakeFilteringModel(const FMakeLinkedFilterViewModelArgs& InArgs) override;
@@ -35,8 +47,10 @@ private:
 	
 	/** The filter bar instance that model edit when they are in linked mode. All created models reference this. */
 	const TSharedRef<FSequencerFilterBar> SharedFilterBar;
-	
 	/** Instance that the active filter bar instances are registered with. This way Sequencer will issue a refilter eventually. */
 	const TSharedRef<FFilterEvaluator> FilterEvaluator;
+	
+	/** Used to determine how many UI systems are using the linked filtering feature. */
+	TArray<TWeakPtr<FLinkedFilterViewModel>> FeatureUsers;
 };
 } // namespace UE::Sequencer
diff --git a/Engine/Source/Editor/Sequencer/Private/SSequencer.cpp b/Engine/Source/Editor/Sequencer/Private/SSequencer.cpp
--- a/Engine/Source/Editor/Sequencer/Private/SSequencer.cpp
+++ b/Engine/Source/Editor/Sequencer/Private/SSequencer.cpp
@@ -160,6 +160,8 @@
 #include "WorkspaceMenuStructure.h"
 #include "WorkspaceMenuStructureModule.h"
 #include "Filters/FilterConfigIdentifiers.h"
+#include "Filters/Linking/Mode/LinkedFilterFactoryViewModel.h"
+#include "Filters/Utils/LinkedFilteringViewUtils.h"
 #include "Filters/Widgets/Linking/SFilterBarSwitcher.h"
 #include "MVVM/ViewModels/OutlinerFilterAreaManager.h"
 #include "MVVM/ViewModels/SequencerEditorViewModelImpl.h"
@@ -366,9 +368,25 @@ void SSequencer::Construct(const FArguments& InArgs, TSharedRef<FSequencer> InSe
 
 	OutlinerAndTrackAreaContent = ConstructOutlinerAndTrackArea(InArgs, TimeSliderArgs);
 
+	const TSharedRef<FSequencerEditorViewModelImpl> RootViewModel = SequencerRef->GetViewModelImpl();
+	const auto IsLinkedToggleButtonVisible = [WeakRootViewModel = RootViewModel.ToWeakPtr()]()
+	{
+		const bool bIsFeatureEnabled = IsLinkedFilteringFeatureEnabled();
+			
+		// The linked filtering toggle button should only be displayed if there is another UI system that can link... if not, the feature is pointless
+		const TSharedPtr<FSequencerEditorViewModelImpl> RootViewModel = WeakRootViewModel.Pin();
+		const TSharedPtr<FLinkedFilterFactoryViewModel> Factory = RootViewModel ? RootViewModel->GetLinkedFilterViewModelFactoryImpl() : nullptr;
+		// Check for more than 1 because Sequencer's instance counts as a user (@see FSequencerEditorViewModelImpl::SequenceFilterInstanceViewModel)
+		const bool bIsAnySystemUsingLinkedFiltering = Factory && Factory->GetFilterFeatureUseCount() > 1;
+			
+		const bool bShouldBeVisible = bIsFeatureEnabled && bIsAnySystemUsingLinkedFiltering;
+		return bShouldBeVisible ? EVisibility::Visible : EVisibility::Collapsed;
+	};
+	
 	FilterRow = SNew(SSequencerSearchAndFilterRow)
 		.ExtenderForAddMenu_Lambda([this]{ return FExtender::Combine(AddMenuExtenders); })
-		.FilterViewModel(SequencerRef->GetViewModelImpl()->GetLinkedFilterViewModelImpl())
+		.FilterViewModel(RootViewModel->GetLinkedFilterViewModelImpl())
+		.FilterToggleButtonVisibility_Lambda(IsLinkedToggleButtonVisible)
 		.Sequencer(SequencerPtr);
 
 	ChildSlot
diff --git a/Engine/Source/Editor/Sequencer/Private/Widgets/Sequencer/SSequencerSearchAndFilterRow.cpp b/Engine/Source/Editor/Sequencer/Private/Widgets/Sequencer/SSequencerSearchAndFilterRow.cpp
--- a/Engine/Source/Editor/Sequencer/Private/Widgets/Sequencer/SSequencerSearchAndFilterRow.cpp
+++ b/Engine/Source/Editor/Sequencer/Private/Widgets/Sequencer/SSequencerSearchAndFilterRow.cpp
@@ -126,7 +126,7 @@ TSharedPtr<SSequencerSearchBox> SSequencerSearchAndFilterRow::GetActiveSearchBox
 TSharedRef<SWidget> SSequencerSearchAndFilterRow::MakeFilterContent(const FArguments& InArgs)
 {
 	const TSharedPtr<FLinkedFilterViewModel> LinkedFilterViewModel = InArgs._FilterViewModel;
-	const auto MakeWidget = [this, &LinkedFilterViewModel](const TSharedRef<FSequencerFilterBar>& InFilterBar)
+	const auto MakeWidget = [this, &LinkedFilterViewModel, &InArgs](const TSharedRef<FSequencerFilterBar>& InFilterBar)
 	{
 		const TSharedRef<FSequencer> SequencerPin = WeakSequencer.Pin().ToSharedRef();
 		const TSharedRef<FHideIsolateShowViewModel> IsolateViewModel = MakeShared<FHideIsolateShowViewModel>(
@@ -145,10 +145,7 @@ TSharedRef<SWidget> SSequencerSearchAndFilterRow::MakeFilterContent(const FArgum
 			.BeforeSearchBox()
 			[
 				SNew(SLinkedFilteringToggleButton, LinkedFilterViewModel.ToSharedRef())
-				.Visibility_Lambda([]
-				{
-					return IsLinkedFilteringFeatureEnabled() ? EVisibility::Visible : EVisibility::Collapsed;
-				})
+				.Visibility(InArgs._FilterToggleButtonVisibility)
 			];
 	};
 	
diff --git a/Engine/Source/Editor/Sequencer/Private/Widgets/Sequencer/SSequencerSearchAndFilterRow.h b/Engine/Source/Editor/Sequencer/Private/Widgets/Sequencer/SSequencerSearchAndFilterRow.h
--- a/Engine/Source/Editor/Sequencer/Private/Widgets/Sequencer/SSequencerSearchAndFilterRow.h
+++ b/Engine/Source/Editor/Sequencer/Private/Widgets/Sequencer/SSequencerSearchAndFilterRow.h
@@ -33,6 +33,9 @@ public:
 		 * Some widgets are stateful, e.g. the search box keeps a search history that should be different between the modes.
 		 */
 		SLATE_ARGUMENT(TSharedPtr<FLinkedFilterViewModel>, FilterViewModel)
+
+		/** The visibility of the toggle button that switches between linked and instanced modes.. */
+		SLATE_ATTRIBUTE(EVisibility, FilterToggleButtonVisibility)
 		
 		/** External extenders for the add button menu. */
 		SLATE_ATTRIBUTE(TSharedPtr<FExtender>, ExtenderForAddMenu)
