Sequencer: Enable CVars for linked fitlering workflow by default:
- Sequencer.CurveEditor.UseLinkedFiltering: Enables the new filter UI in Sequencer's Curve Editor
- Sequencer.UseLinkedFiltering: Enables the toggle button in Sequencer

#Synced-CL 51881412

diff --git a/Engine/Source/Editor/Sequencer/Private/CurveEditor/Views/CurveEditorWidgetOwner.cpp b/Engine/Source/Editor/Sequencer/Private/CurveEditor/Views/CurveEditorWidgetOwner.cpp
--- a/Engine/Source/Editor/Sequencer/Private/CurveEditor/Views/CurveEditorWidgetOwner.cpp
+++ b/Engine/Source/Editor/Sequencer/Private/CurveEditor/Views/CurveEditorWidgetOwner.cpp
@@ -30,8 +30,9 @@ namespace UE::Sequencer
 {
 namespace PresenterDetail
 {
+// TODO UE-371833: Remove CVar in 5.9 assuming the system is stable.
 static TAutoConsoleVariable<bool> CVarUseLinkedFilterUI(
-	TEXT("Sequencer.CurveEditor.UseLinkedFiltering"), false, TEXT("Whether to use the new linked filtering UI in Curve Editor")
+	TEXT("Sequencer.CurveEditor.UseLinkedFiltering"), true, TEXT("Whether to use the new linked filtering UI in Curve Editor")
 	);
 
 static FCurveEditorWidgets MakePresenterWidgets(
diff --git a/Engine/Source/Editor/Sequencer/Private/Filters/Utils/LinkedFilteringViewUtils.cpp b/Engine/Source/Editor/Sequencer/Private/Filters/Utils/LinkedFilteringViewUtils.cpp
--- a/Engine/Source/Editor/Sequencer/Private/Filters/Utils/LinkedFilteringViewUtils.cpp
+++ b/Engine/Source/Editor/Sequencer/Private/Filters/Utils/LinkedFilteringViewUtils.cpp
@@ -31,8 +31,9 @@ bool IsLinkedFilteringActionChecked(TWeakPtr<ILinkedFilterViewModel> InWeakViewM
 	return ViewModelPin && ViewModelPin->GetFilterMode() == ELinkedFilterMode::Linked;
 }
 
+// TODO UE-371833: Remove CVar in 5.9 assuming the system is stable.
 TAutoConsoleVariable<bool> CVarSequencerUseLinkedFiltering(
-	TEXT("Sequencer.UseLinkedFiltering"), false, TEXT("Whether to use the new linked filtering system")
+	TEXT("Sequencer.UseLinkedFiltering"), true, TEXT("Whether to use the new linked filtering system")
 	);
 
 bool IsLinkedFilteringFeatureEnabled()
