Sequencer: Fix SSequencerSection background painting to only skip lanes with their own section model. This fixes issues for painting section backgrounds for decoration sections with channels.

#Synced-CL 51964473

diff --git a/Engine/Source/Editor/Sequencer/Private/SSequencerSection.cpp b/Engine/Source/Editor/Sequencer/Private/SSequencerSection.cpp
--- a/Engine/Source/Editor/Sequencer/Private/SSequencerSection.cpp
+++ b/Engine/Source/Editor/Sequencer/Private/SSequencerSection.cpp
@@ -731,15 +731,26 @@ struct FSequencerSectionPainterImpl : FSequencerSectionPainter
 		TArray<TSharedPtr<STrackLane>> EmptyChildLanes;
 		Algo::Transform(OwningTrackLane->GetChildLanes(), EmptyChildLanes, [](const TWeakPtr<STrackLane>& In){ return In.Pin(); });
 
-		// Remove any invalid track lanes, or track lanes that already have widgets
-		// (widgets handle their own rendering, so we shouldn't draw backgrounds over them)
+		// Remove invalid track lanes and lanes that have their own section widgets
+		// (sections handle their own background painting). Lanes with only channel/key
+		// area widgets still need background from this parent section.
 		for (int32 Index = EmptyChildLanes.Num()-1; Index >= 0; --Index)
 		{
 			if (EmptyChildLanes[Index])
 			{
-				// If this track lane has any widgets at all, don't draw an empty background for it
-				// This includes decoration channel widgets which are not descendants of the section
-				if (EmptyChildLanes[Index]->GetAllWidgets().Num() > 0)
+				bool bHasOwnSection = false;
+				for (const auto& Pair : EmptyChildLanes[Index]->GetAllWidgets())
+				{
+					if (TSharedPtr<FViewModel> Model = Pair.Key.Pin())
+					{
+						if (Model->IsA<FSectionModel>())
+						{
+							bHasOwnSection = true;
+							break;
+						}
+					}
+				}
+				if (bHasOwnSection)
 				{
 					EmptyChildLanes.RemoveAtSwap(Index, EAllowShrinking::No);
 				}
