Adding missing includes.

#Synced-CL 52214067

diff --git a/Engine/Plugins/VirtualProduction/Rivermax/RivermaxCore/Source/RivermaxCore/Private/RTPHeader.h b/Engine/Plugins/VirtualProduction/Rivermax/RivermaxCore/Source/RivermaxCore/Private/RTPHeader.h
--- a/Engine/Plugins/VirtualProduction/Rivermax/RivermaxCore/Source/RivermaxCore/Private/RTPHeader.h
+++ b/Engine/Plugins/VirtualProduction/Rivermax/RivermaxCore/Source/RivermaxCore/Private/RTPHeader.h
@@ -4,6 +4,7 @@
 
 #include "CoreTypes.h"
 #include "Containers/ArrayView.h"
+#include "Containers/UnrealString.h"
 
 #if UE_BUILD_DEBUG
 	#define RIVERMAX_PACKET_DEBUG 1
diff --git a/Engine/Source/Editor/Sequencer/Public/Filters/SequencerFilterBase.h b/Engine/Source/Editor/Sequencer/Public/Filters/SequencerFilterBase.h
--- a/Engine/Source/Editor/Sequencer/Public/Filters/SequencerFilterBase.h
+++ b/Engine/Source/Editor/Sequencer/Public/Filters/SequencerFilterBase.h
@@ -3,6 +3,7 @@
 #pragma once
 
 #include "Filters/FilterBase.h"
+#include "Filters/ISequencerFilterBar.h"
 #include "Framework/Commands/UIAction.h"
 #include "Framework/Commands/UICommandInfo.h"
 #include "Framework/Commands/InputChord.h"
@@ -10,7 +11,6 @@
 #include "MovieSceneSequence.h"
 
 class ISequencer;
-class ISequencerFilterBar;
 
 /** Sequencer base filter for all Sequencer filters (Sequencer and Sequence Navigator) */
 template<typename InFilterType>
