Sequencer: Adjust filtering algorithm to improve performance when filtering using the Selected filter. Future improvements can be made to other filters later.

The main idea is to avoid search recoursing the FViewModel tree hiearchy needlessly. Filters can now specify whether their result applies to all children, too. For example, for the selection filter, we need not iterate all children of a selected node: they should all be shown so it is pointless to rerun filters on those nodes.

Main changes:
- EItemFilterState: Include, Exclude, DoNotCare.
- FSequencerTrackFilter::Evaluate: This function allows filters to whether the filter in (include), out (exclude), or do not affect an item AND whether the result should be applied to all children.
- FSequencerTrackFilter_Selected::OnPreFilter now pre-iterates the entire hierarchy once to determine the view models to include. This is massive performance boost: before the algorithm would iterate each tree item, and FSequencerTrackFilter_Selected would reiterate up the entire parent hierarchy to check whether one of its parents is selected.

The old algorithm continues to exist and can be enabled using "Sequencer.UseNewFilteringAlgorithm 0" console variable. The new algorithm is enabled by default. The point of this CVar is so in case anything is broken in the new filtering algorith, animators can revert back to the old algorithm locally while we address any potential bugs.

On a simple case with only the selected filter, 2 SM actors, 1 Mannequin, and 2 dragons, filtering went from 470ms to only 1.77ms.

#Synced-CL 51881412

diff --git a/Engine/Source/Editor/Sequencer/Private/Filters/Evaluation/FilterEvaluator.cpp b/Engine/Source/Editor/Sequencer/Private/Filters/Evaluation/FilterEvaluator.cpp
--- a/Engine/Source/Editor/Sequencer/Private/Filters/Evaluation/FilterEvaluator.cpp
+++ b/Engine/Source/Editor/Sequencer/Private/Filters/Evaluation/FilterEvaluator.cpp
@@ -40,11 +40,14 @@ bool FFilterEvaluator::UpdateFilters(bool bForce)
 	{
 		return false;
 	}
+	
+	TRACE_CPUPROFILER_EVENT_SCOPE(FFilterEvaluator::UpdateFilters);
 
 	bool bHasFilterDataChanged = false; 
 	for (FFilterBarData& FilterBarData : RegisteredFilterBars)
 	{
 		const TSharedRef<FSequencerFilterBar>& FilterBar = FilterBarData.FilterBar;
+		TRACE_CPUPROFILER_EVENT_SCOPE_TEXT(FilterBar->GetIdentifier());
 		
 		const FSequencerFilterData PreviousFilterData = FilterBar->GetFilterData();
 		const FSequencerFilterData& FilterData = FilterBar->FilterNodes();
diff --git a/Engine/Source/Editor/Sequencer/Private/Filters/Filters/SequencerTrackFilterBase.cpp b/Engine/Source/Editor/Sequencer/Private/Filters/Filters/SequencerTrackFilterBase.cpp
--- a/Engine/Source/Editor/Sequencer/Private/Filters/Filters/SequencerTrackFilterBase.cpp
+++ b/Engine/Source/Editor/Sequencer/Private/Filters/Filters/SequencerTrackFilterBase.cpp
@@ -23,6 +23,13 @@ ISequencerTrackFilters& FSequencerTrackFilter::GetFilterInterface() const
 	return static_cast<ISequencerTrackFilters&>(FilterInterface);
 }
 
+FFilterResult FSequencerTrackFilter::Evaluate(FSequencerTrackFilterType InItem) const
+{
+	// This API was added in 5.8. We allow filters to only implement PassesFilter.
+	const bool bPassesFilter = PassesFilter(InItem);
+	return FFilterResult(bPassesFilter ? EItemFilterState::Include : EItemFilterState::Exclude);
+}
+
 UMovieSceneSequence* FSequencerTrackFilter::GetFocusedMovieSceneSequence() const
 {
 	return GetSequencer().GetFocusedMovieSceneSequence();
diff --git a/Engine/Source/Editor/Sequencer/Private/Filters/Filters/SequencerTrackFilter_Selected.cpp b/Engine/Source/Editor/Sequencer/Private/Filters/Filters/SequencerTrackFilter_Selected.cpp
--- a/Engine/Source/Editor/Sequencer/Private/Filters/Filters/SequencerTrackFilter_Selected.cpp
+++ b/Engine/Source/Editor/Sequencer/Private/Filters/Filters/SequencerTrackFilter_Selected.cpp
@@ -5,6 +5,7 @@
 #include "Filters/SequencerFilterBar.h"
 #include "Filters/SequencerTrackFilterCommands.h"
 #include "Filters/SequencerTrackFilterSelectedExtension.h"
+#include "HAL/IConsoleManager.h"
 #include "ISequencer.h"
 #include "MVVM/Selection/Selection.h"
 #include "MVVM/ViewModels/SequencerEditorViewModel.h"
@@ -15,8 +16,6 @@ using namespace UE::Sequencer;
 
 #define LOCTEXT_NAMESPACE "SequencerTrackFilter_Selected"
 
-
-
 namespace UE::Sequencer::SelectionDetail
 {
 static TArray<TWeakObjectPtr<USequencerTrackFilterSelectedExtension>> GetSelectionExtensions()
@@ -42,7 +41,6 @@ static bool IsPartOfHierarchy(const FViewModelPtr& InItem, const FViewModelPtr&
 	{
 		return false;
 	}
-
 	for (FParentFirstChildIterator It(InSelectedNode, /*bIncludeThis=*/true); It; ++It)
 	{
 		if (*It == InItem)
@@ -53,7 +51,6 @@ static bool IsPartOfHierarchy(const FViewModelPtr& InItem, const FViewModelPtr&
 	
 	return false;
 }
-
 /** @return Whether InItem is a child of the FViewModels enumerated by USequencerTrackFilterSelectedExtension. */
 static bool IsPartOfExtendedSelection(
 	const FViewModelPtr& InItem, 
@@ -63,20 +60,20 @@ static bool IsPartOfExtendedSelection(
 {
 	using EBreakBehavior = USequencerTrackFilterSelectedExtension::EBreakBehavior;
 	const auto ProcessExtension = [&InItem, &InFilter](const TWeakObjectPtr<USequencerTrackFilterSelectedExtension>& WeakExtension)
-	{
-		if (!WeakExtension.IsValid())
 		{
-			return false;
-		}
+			if (!WeakExtension.IsValid())
+			{
+				return false;
+			}
 		
-		bool bIsPartOfSelection = false;
-		WeakExtension->EnumerateViewModelsConsideredAsSelected(InFilter, [&InItem, &bIsPartOfSelection](const FViewModelPtr& ConsiderSelected)
-		{
-			bIsPartOfSelection = IsPartOfHierarchy(InItem, ConsiderSelected);
-			return bIsPartOfSelection ? EBreakBehavior::Break : EBreakBehavior::Continue;
-		});
-		return bIsPartOfSelection;
-	};
+			bool bIsPartOfSelection = false;
+			WeakExtension->EnumerateViewModelsConsideredAsSelected(InFilter, [&InItem, &bIsPartOfSelection](const FViewModelPtr& ConsiderSelected)
+			{
+				bIsPartOfSelection = IsPartOfHierarchy(InItem, ConsiderSelected);
+				return bIsPartOfSelection ? EBreakBehavior::Break : EBreakBehavior::Continue;
+			});
+			return bIsPartOfSelection;
+		};
 	
 	return Algo::AnyOf(InWeakExtensions, ProcessExtension);
 }
@@ -89,14 +86,14 @@ FSequencerTrackFilter_Selected::FSequencerTrackFilter_Selected(ISequencerTrackFi
 
 FSequencerTrackFilter_Selected::~FSequencerTrackFilter_Selected()
 {
-	UnbindSelectionChanged();
+	UnbindEvents();
 }
 
 void FSequencerTrackFilter_Selected::BindSelectionChanged()
 {
 	if (!OnSelectionChangedHandle.IsValid())
 	{
-		OnSelectionChangedHandle = USelection::SelectionChangedEvent.AddRaw(this, &FSequencerTrackFilter_Selected::OnSelectionChanged);
+		OnSelectionChangedHandle = USelection::SelectionChangedEvent.AddRaw(this, &FSequencerTrackFilter_Selected::OnViewportSelectionChanged);
 	}
 }
 
@@ -130,14 +127,24 @@ void FSequencerTrackFilter_Selected::ActiveStateChanged(const bool bInActive)
 
 	if (bInActive)
 	{
-		BindSelectionChanged();
+		BindEvents();
 	}
 	else
 	{
-		UnbindSelectionChanged();
+		UnbindEvents();
 	}
 }
 
+FFilterResult FSequencerTrackFilter_Selected::Evaluate(FSequencerTrackFilterType InItem) const
+{
+	const EFilterResultFlags* Flags = IncludedViewModels.Find(InItem);
+	const bool bPassesFilter = Flags != nullptr;
+	
+	const EItemFilterState ResultState = bPassesFilter ? EItemFilterState::Include : EItemFilterState::Exclude;
+	const EFilterResultFlags ResultFlags = Flags ? *Flags : EFilterResultFlags::None;
+	return FFilterResult(ResultState, ResultFlags);
+}
+
 FText FSequencerTrackFilter_Selected::GetDisplayName() const
 {
 	return LOCTEXT("SequencerTrackFilter_Selected", "Selected");
@@ -148,6 +155,69 @@ FSlateIcon FSequencerTrackFilter_Selected::GetIcon() const
 	return FSlateIcon(FAppStyle::GetAppStyleSetName(), TEXT("Icons.SelectInViewport"));
 }
 
+void FSequencerTrackFilter_Selected::OnPreFilter()
+{
+	if (!bNeedsPrefilterRebuild)
+	{
+		return;
+	}
+	
+	TRACE_CPUPROFILER_EVENT_SCOPE(FSequencerTrackFilter_Selected::OnPreFilter);
+	bNeedsPrefilterRebuild = false;
+	IncludedViewModels.Empty(IncludedViewModels.Num());
+	ON_SCOPE_EXIT{ IncludedViewModels.Shrink(); };
+	
+	const TSharedPtr<FSequencerEditorViewModel> RootViewModel = GetSequencer().GetViewModel();
+	const TSharedPtr<FSequencerSelection> SequencerSelection = RootViewModel ? RootViewModel->GetSelection() : nullptr;
+	if (!SequencerSelection.IsValid())
+	{
+		return;
+	}
+	
+	const auto ProcessSelected = [this](const FViewModelPtr& ViewModel)
+	{
+		if (IncludedViewModels.Contains(ViewModel))
+		{
+			return;
+		}
+			
+		// If "Foo.A" is selected, make "Foo.A.Child", etc. filtered in...
+		IncludedViewModels.Add(ViewModel, EFilterResultFlags::ApplyResultToChildren);
+	
+		for (const FViewModelPtr& Child : FParentFirstChildIterator(ViewModel))
+		{
+			// ... an if the filtered is asked whether "Foo.A.Child" is filtered in, it should also say, yes.
+			IncludedViewModels.Add(Child, EFilterResultFlags::ApplyResultToChildren);
+		}
+			
+		for (const FViewModelPtr& Parent : FParentModelIterator(ViewModel))
+		{
+			// Parent should be visible but do not tell them to propagate to all their children. 
+			// If Foo.A is selected, then Foo should be filtered in, but not Foo.B.
+			IncludedViewModels.FindOrAdd(Parent);
+		}
+	};
+
+	for (const FViewModelPtr& ViewModel : SequencerSelection->Outliner)
+	{
+		ProcessSelected(ViewModel);
+	}
+	
+	for (const TWeakObjectPtr<USequencerTrackFilterSelectedExtension>& Extension : CachedWeakExtensions)
+	{
+		if (!Extension.IsValid())
+		{
+			continue;
+		}
+		
+		Extension->EnumerateViewModelsConsideredAsSelected(*this, [&ProcessSelected](const FViewModelPtr& ViewModel)
+		{
+			ProcessSelected(ViewModel);
+			return USequencerTrackFilterSelectedExtension::EBreakBehavior::Continue;
+		});
+	}
+}
+
 FString FSequencerTrackFilter_Selected::GetName() const
 {
 	return StaticName();
@@ -155,35 +225,70 @@ FString FSequencerTrackFilter_Selected::GetName() const
 
 bool FSequencerTrackFilter_Selected::PassesFilter(FSequencerTrackFilterType InItem) const
 {
-	const TSharedPtr<FSequencerEditorViewModel> ViewModel = GetSequencer().GetViewModel();
-	if (!ViewModel.IsValid())
+	// TODO UE-371829: Remove this CVar
+	const IConsoleVariable* Variable = IConsoleManager::Get().FindConsoleVariable(TEXT("Sequencer.UseNewFilteringAlgorithm"));
+	const bool bUseOldAlgorithm = !ensure(Variable) || !Variable->GetBool();
+	if (bUseOldAlgorithm)
 	{
-		return false;
+		const TSharedPtr<FSequencerEditorViewModel> ViewModel = GetSequencer().GetViewModel();
+		if (!ViewModel.IsValid())
+		{
+			return false;
+		}
+ 
+		const TSharedPtr<FSequencerSelection> SequencerSelection = ViewModel->GetSelection();
+		if (!SequencerSelection.IsValid())
+		{
+			return false;
+		}
+		const bool bIsPartOfSequencerSelection = Algo::AnyOf(SequencerSelection->Outliner, [InItem](const FViewModelPtr& Selected)
+		{
+			return SelectionDetail::IsPartOfHierarchy(InItem, Selected);
+		});
+		return bIsPartOfSequencerSelection 
+			// Allow external logic to make us pretend like additional items are selected without actually modifying FSequencerSelection::Outliner.
+			|| SelectionDetail::IsPartOfExtendedSelection(InItem, *this, CachedWeakExtensions);
 	}
+	
+	return IncludedViewModels.Contains(InItem);
+}
 
-	const TSharedPtr<FSequencerSelection> SequencerSelection = ViewModel->GetSelection();
-	if (!SequencerSelection.IsValid())
+void FSequencerTrackFilter_Selected::BindEvents()
+{
+	BindSelectionChanged();
+	
+	const TSharedPtr<FSequencerEditorViewModel> RootViewModel = GetSequencer().GetViewModel();
+	const TSharedPtr<FSequencerSelection> SequencerSelection = RootViewModel ? RootViewModel->GetSelection() : nullptr;
+	if (SequencerSelection)
 	{
-		return false;
+		SequencerSelection->OnChanged.AddRaw(this, &FSequencerTrackFilter_Selected::OnSequencerSelectionChanged);
 	}
+}
 
-	const bool bIsPartOfSequencerSelection = Algo::AnyOf(SequencerSelection->Outliner, [InItem](const FViewModelPtr& Selected)
+void FSequencerTrackFilter_Selected::UnbindEvents()
+{
+	UnbindSelectionChanged();
+	
+	const TSharedPtr<FSequencerEditorViewModel> RootViewModel = GetSequencer().GetViewModel();
+	const TSharedPtr<FSequencerSelection> SequencerSelection = RootViewModel ? RootViewModel->GetSelection() : nullptr;
+	if (SequencerSelection)
 	{
-		return SelectionDetail::IsPartOfHierarchy(InItem, Selected);
-	});
-	return bIsPartOfSequencerSelection 
-		// Allow external logic to make us pretend like additional items are selected without actually modifying FSequencerSelection::Outliner.
-		|| SelectionDetail::IsPartOfExtendedSelection(InItem, *this, CachedWeakExtensions);
+		SequencerSelection->OnChanged.RemoveAll(this);
+	}
 }
 
-void FSequencerTrackFilter_Selected::OnSelectionChanged(UObject* const InObject)
+void FSequencerTrackFilter_Selected::OnSequencerSelectionChanged()
 {
-	FilterInterface.RequestFilterUpdate();
+	bNeedsPrefilterRebuild = true;
+	// TODO UE-371115: For now, we achieve this ticket behavior partially by not refiltering. However, once the ticket is completed, this should be uncommented.
+	// FilterInterface.RequestFilterUpdate();
+	// Achieved partially because any subsequent filter reevaluation, like updating the search bar and resetting it, will hide previously multi-selected items.
 }
 
-void FSequencerTrackFilter_Selected::ToggleShowOnlySelectedTracks()
+void FSequencerTrackFilter_Selected::OnViewportSelectionChanged(UObject* const InObject)
 {
-	SetActive(!IsActive());
+	bNeedsPrefilterRebuild = true;
+	FilterInterface.RequestFilterUpdate();
 }
 
 #undef LOCTEXT_NAMESPACE
diff --git a/Engine/Source/Editor/Sequencer/Private/Filters/SequencerFilterBar.cpp b/Engine/Source/Editor/Sequencer/Private/Filters/SequencerFilterBar.cpp
--- a/Engine/Source/Editor/Sequencer/Private/Filters/SequencerFilterBar.cpp
+++ b/Engine/Source/Editor/Sequencer/Private/Filters/SequencerFilterBar.cpp
@@ -33,6 +33,7 @@
 #include "Filters/SequencerFilterAreaConfig.h"
 #include "Stats/StatsMisc.h"
 #include "UObject/UObjectIterator.h"
+#include "Utils/FilterEvaluation.h"
 #include "Utils/HideIsolateViewModelUtils.h"
 
 using namespace UE::Sequencer;
@@ -1387,7 +1388,9 @@ UWorld* FSequencerFilterBar::GetWorld() const
 
 void FSequencerFilterBar::PreFilter() const
 {
-	auto PreFilter = [this](const TSharedRef<FSequencerTrackFilter>& InFilter)
+	TRACE_CPUPROFILER_EVENT_SCOPE(FSequencerFilterBar::PreFilter);
+	
+	const auto PreFilter = [this](const TSharedRef<FSequencerTrackFilter>& InFilter)
 		{
 			if (IsFilterActive(InFilter))
 			{
@@ -1395,7 +1398,7 @@ void FSequencerFilterBar::PreFilter() const
 			}
 		};
 
-	auto PreFilterCollection = [PreFilter](const TSharedRef<FSequencerTrackFilterCollection>& InCollection)
+	const auto PreFilterCollection = [PreFilter](const TSharedRef<FSequencerTrackFilterCollection>& InCollection)
 		{
 			InCollection->ForEachFilter([PreFilter]
 				(const TSharedRef<FSequencerTrackFilter>& InFilter)
@@ -1416,22 +1419,20 @@ void FSequencerFilterBar::PreFilter() const
 	PreFilter(ModifiedFilter);
 }
 
+// TODO UE-371829: Remove this CVar: Also look for "Sequencer.UseNewFilteringAlgorithm" because it passed to FindConsoleVariable
+static TAutoConsoleVariable<bool> CVarUseNewFilteringAlgorithm(TEXT("Sequencer.UseNewFilteringAlgorithm"), true, TEXT(""));
+
 const FSequencerFilterData& FSequencerFilterBar::FilterNodes()
 {
-	//SCOPE_LOG_TIME_IN_SECONDS(TEXT("FSequencerFilterBar::FilterNodes()"), nullptr);
-
-	// Allow all filters to do any pre-processing before filtering
-	PreFilter();
-
-	// Reset all filter data
-	FilterData.Reset();
-
-	// Loop through all nodes and filter recursively
-	const bool bHasActiveFilter = HasAnyFilterActive();
-	for (const TViewModelPtr<IOutlinerExtension>& RootNode : Sequencer.GetNodeTree()->GetRootNodes())
+	if (CVarUseNewFilteringAlgorithm.GetValueOnAnyThread())
 	{
-		FilterNodesRecursive(bHasActiveFilter, RootNode);
+		FilterNodesNew();
 	}
+	else
+	{
+		FilterNodesOld();
+	}
+	
 
 	// Always filter in spacer node
 	if (const FSequenceModel* const SequenceModel = Sequencer.GetNodeTree()->GetRootNode()->CastThis<FSequenceModel>())
@@ -1501,7 +1502,23 @@ FString FSequencerFilterBar::GenerateTextFilterStringFromEnabledFilters() const
 	return GeneratedFilterString;
 }
 
-bool FSequencerFilterBar::FilterNodesRecursive(const bool bInHasActiveFilter, const TViewModelPtr<IOutlinerExtension>& InStartNode)
+void FSequencerFilterBar::FilterNodesOld()
+{
+	// Allow all filters to do any pre-processing before filtering
+	PreFilter();
+	
+	// Reset all filter data
+	FilterData.Reset();
+	
+	// Loop through all nodes and filter recursively
+	const bool bHasActiveFilter = HasAnyFilterActive();
+	for (const TViewModelPtr<IOutlinerExtension>& RootNode : Sequencer.GetNodeTree()->GetRootNodes())
+	{
+		FilterNodesRecursiveOld(bHasActiveFilter, RootNode);
+	}
+}
+
+bool FSequencerFilterBar::FilterNodesRecursiveOld(const bool bInHasActiveFilter, const TViewModelPtr<IOutlinerExtension>& InStartNode)
 {
 	TRACE_CPUPROFILER_EVENT_SCOPE(FSequencerFilterBar::FilterNodesRecursive);
 	
@@ -1512,57 +1529,54 @@ bool FSequencerFilterBar::FilterNodesRecursive(const bool bInHasActiveFilter, co
 	 * - Hidden/Isolated tracks will take precedence over common filters
 	 * - Can hide sub tracks of isolated tracks
 	 */
-
 	bool bAnyChildPassed = false;
-
+	
 	// Child nodes should always be processed, as they may force their parents to pass
 	for (const TViewModelPtr<IOutlinerExtension>& Node : InStartNode.AsModel()->GetChildrenOfType<IOutlinerExtension>())
 	{
-		if (FilterNodesRecursive(bInHasActiveFilter, Node))
+		if (FilterNodesRecursiveOld(bInHasActiveFilter, Node))
 		{
 			bAnyChildPassed = true;
 		}
 	}
-
+	
 	// Increment the total node count so we can remove the code to loop again just to count
 	FilterData.IncrementTotalNodeCount();
-
 	// Early out if no filter
 	if (!bInHasActiveFilter)
 	{
 		FilterData.FilterInNode(InStartNode);
 		return false;
 	}
-
+	
 	const USequencerSettings* const SequencerSettings = Sequencer.GetSequencerSettings();
 	if (!SequencerSettings)
 	{
 		return false;
 	}
-
+	
 	// Pinning overrides all other filters
 	if (EnumHasAnyFlags(Flags, EFilterFlags::PinnedItemsCanSkipFiltering) && !SequencerSettings->GetIncludePinnedInFilter())
 	{
 		const TSharedPtr<IPinnableExtension> Pinnable = InStartNode.AsModel()->FindAncestorOfType<IPinnableExtension>(true);
 		if (Pinnable.IsValid() && Pinnable->IsPinned())
 		{
-			FilterData.FilterInParentChildNodes(InStartNode, true, true, true);
+			FilterData.FilterInParentChildNodes(InStartNode, ENodeVisitFlags::IncludeAll);
 			return true;
 		}
 	}
-
+	
 	const bool bPassedTextFilter = !TextFilter->IsActive() || TextFilter->PassesFilter(InStartNode);
 	const bool bPassedHideIsolateFilter = !HideIsolateFilter->IsActive() || HideIsolateFilter->PassesFilter(InStartNode);
 	const bool bPassedAnyCommonFilters = PassesAnyCommonFilter(InStartNode);
 	const bool bPassedInternalFilters = !AnyInternalFilterActive() || PassesAllInternalFilters(InStartNode);
 	const bool bPassedAnyCustomTextFilters = PassesAllCustomTextFilters(InStartNode);
-
 	const bool bAllFiltersPassed = bPassedTextFilter
 		&& bPassedHideIsolateFilter
 		&& bPassedAnyCommonFilters
 		&& bPassedInternalFilters
 		&& bPassedAnyCustomTextFilters;
-
+	
 	if (bAllFiltersPassed || bAnyChildPassed)
 	{
 		// Don't expand if a mouse capture is effecting TrackValueChanged events to avoid keys
@@ -1572,6 +1586,110 @@ bool FSequencerFilterBar::FilterNodesRecursive(const bool bInHasActiveFilter, co
 		{
 			SetTrackParentsExpanded(InStartNode.ImplicitCast(), true);
 		}
+		FilterData.FilterInNodeWithAncestors(InStartNode);
+		return true;
+	}
+	// After child nodes are processed, fail anything that didn't pass
+	FilterData.FilterOutNode(InStartNode);
+	return false;
+}
+
+void FSequencerFilterBar::FilterNodesNew()
+{
+	// Reset all filter data
+	FilterData.Reset();
+	
+	if (HasAnyFilterActive())
+	{
+		const TArray<TNotNull<FSequencerTrackFilter*>> Filters = BuildActiveFilters();
+		
+		for (FSequencerTrackFilter* Filter : Filters)
+		{
+			Filter->OnPreFilter();
+		}
+		
+		// Loop through all nodes and filter recursively
+		for (const TViewModelPtr<IOutlinerExtension>& RootNode : Sequencer.GetNodeTree()->GetRootNodes())
+		{
+			FilterNodesRecursiveNew(RootNode, Filters, {});
+		}
+	}
+	else
+	{
+		for (const TViewModelPtr<IOutlinerExtension>& RootNode : Sequencer.GetNodeTree()->GetRootNodes())
+		{
+			FilterData.FilterInParentChildNodes(RootNode, ENodeVisitFlags::IncludeSelf | ENodeVisitFlags::IncludeChildren);
+		}
+	}
+}
+
+bool FSequencerFilterBar::FilterNodesRecursiveNew(
+	const TViewModelPtr<IOutlinerExtension>& InStartNode,
+	TConstArrayView<TNotNull<FSequencerTrackFilter*>> InFilters,
+	const FRecursiveFilterOverrides& InFilterOverrides
+	)
+{
+	TRACE_CPUPROFILER_EVENT_SCOPE(FSequencerFilterBar::FilterNodesRecursive);
+	
+	/**
+	 * Main Filtering Logic
+	 *
+	 * - Pinning overrides all other filters
+	 * - Hidden/Isolated tracks will take precedence over common filters
+	 * - Can hide sub tracks of isolated tracks
+	 */
+
+	const USequencerSettings* const SequencerSettings = Sequencer.GetSequencerSettings();
+
+	// Pinning overrides all other filters
+	if (EnumHasAnyFlags(Flags, EFilterFlags::PinnedItemsCanSkipFiltering) && SequencerSettings && !SequencerSettings->GetIncludePinnedInFilter())
+	{
+		const TSharedPtr<IPinnableExtension> Pinnable = InStartNode.AsModel()->FindAncestorOfType<IPinnableExtension>(true);
+		if (Pinnable.IsValid() && Pinnable->IsPinned())
+		{
+			const ENodeVisitFlags VisitFlags = ENodeVisitFlags::IncludeSelf | ENodeVisitFlags::IncludeChildren;
+			FilterData.FilterInParentChildNodes(InStartNode, VisitFlags);
+			return true;
+		}
+	}
+	
+	const FFilterEvaluation EvaluationResult = EvaluateFiltersOnItem(InStartNode, InFilters, InFilterOverrides);
+	const bool bAllFiltersPassed = EvaluationResult.ItemResult != EItemFilterState::Exclude;
+	
+	bool bAnyChildPassed = false;
+	const bool bHasOverrideForAllFilters = EvaluationResult.NewRecursiveOverrides.Overrides.Num() == InFilters.Num();
+	if (bHasOverrideForAllFilters)
+	{
+		bAnyChildPassed = bAllFiltersPassed;
+		const ENodeVisitFlags VisitFlags = ENodeVisitFlags::IncludeSelf | ENodeVisitFlags::IncludeChildren;
+		const int32 Num = bAllFiltersPassed 
+			? FilterData.FilterInParentChildNodes(InStartNode, VisitFlags) 
+			: FilterData.FilterOutParentChildNodes(InStartNode, VisitFlags);
+		FilterData.AddTotalNodeCount(Num);
+	}
+	else
+	{
+		FilterData.IncrementTotalNodeCount();
+		
+		// Child nodes should always be processed, as they may force their parents to pass
+		for (const TViewModelPtr<IOutlinerExtension>& Node : InStartNode.AsModel()->GetChildrenOfType<IOutlinerExtension>())
+		{
+			if (FilterNodesRecursiveNew(Node, InFilters, EvaluationResult.NewRecursiveOverrides))
+			{
+				bAnyChildPassed = true;
+			}
+		}
+	}
+	
+	if (bAllFiltersPassed || bAnyChildPassed)
+	{
+		// Don't expand if a mouse capture is effecting TrackValueChanged events to avoid keys
+		// being hidden and expanding tracks while dragging
+		if (SequencerSettings && SequencerSettings->GetAutoExpandNodesOnFilterPass()
+			&& !FSlateApplication::Get().HasAnyMouseCaptor())
+		{
+			SetTrackParentsExpanded(InStartNode.ImplicitCast(), true);
+		}
 
 		FilterData.FilterInNodeWithAncestors(InStartNode);
 		return true;
@@ -1582,6 +1700,49 @@ bool FSequencerFilterBar::FilterNodesRecursive(const bool bInHasActiveFilter, co
 	return false;
 }
 
+TArray<TNotNull<FSequencerTrackFilter*>> FSequencerFilterBar::BuildActiveFilters()
+{
+	TArray<TNotNull<FSequencerTrackFilter*>> Filters;
+	
+	if (TextFilter->IsActive())
+	{
+		Filters.Emplace(&TextFilter.Get());
+	}
+	
+	if (HideIsolateFilter->IsActive())
+	{
+		Filters.Emplace(&HideIsolateFilter.Get());
+	}
+	
+	CommonFilters->ForEachFilter([this, &Filters](const TSharedRef<FSequencerTrackFilter>& InFilter)
+	{
+		if (IsFilterActive(InFilter))
+		{
+			Filters.Emplace(&InFilter.Get());
+		}
+		return true;
+	});
+	
+	InternalFilters->ForEachFilter([this, &Filters](const TSharedRef<FSequencerTrackFilter>& InFilter)
+	{
+		if (IsFilterActive(InFilter))
+		{
+			Filters.Emplace(&InFilter.Get());
+		}
+		return true;
+	});
+	
+	for (const TSharedRef<FSequencerTrackFilter_CustomText>& Filter : CustomTextFilters)
+	{
+		if (IsFilterActive(Filter))
+		{
+			Filters.Emplace(&Filter.Get());
+		}
+	}
+	
+	return Filters;
+}
+
 void FSequencerFilterBar::HandleTextFilterChanged()
 {
 	OnTextFilterTextChangedEvent.Broadcast();
diff --git a/Engine/Source/Editor/Sequencer/Private/Filters/SequencerFilterBar.h b/Engine/Source/Editor/Sequencer/Private/Filters/SequencerFilterBar.h
--- a/Engine/Source/Editor/Sequencer/Private/Filters/SequencerFilterBar.h
+++ b/Engine/Source/Editor/Sequencer/Private/Filters/SequencerFilterBar.h
@@ -26,7 +26,9 @@ class SFilterBarIsolateHideShow;
 class SSequencerFilterBar;
 class SSequencerSearchBox;
 enum class EFilterBarLayout : uint8;
-namespace UE::Sequencer { class IOutlinerExtension; }
+namespace UE::Sequencer {
+struct FRecursiveFilterOverrides;
+class IOutlinerExtension; }
 namespace UE::Sequencer { struct FFilterBarInitArgs; }
 
 enum class ESequencerFilterChange : uint8
@@ -321,8 +323,21 @@ private:
 	/** Binds all commands that the filter bar or its filters handle. */
 	void BindCommands();
 	
+	/** Old algorithm for filtering nodes */
+	void FilterNodesOld();
 	/** Do not call directly! Should only be called by FilterNodes(). */
-	bool FilterNodesRecursive(const bool bInHasActiveFilter, const UE::Sequencer::TViewModelPtr<UE::Sequencer::IOutlinerExtension>& InStartNode);
+	bool FilterNodesRecursiveOld(const bool bInHasActiveFilter, const UE::Sequencer::TViewModelPtr<UE::Sequencer::IOutlinerExtension>& InStartNode);
+	
+	/** New algorithm for filtering nodes. */
+	void FilterNodesNew();
+	/** @return Filter that are active and should be evaluated. */
+	TArray<TNotNull<FSequencerTrackFilter*>> BuildActiveFilters();
+	/** Do not call directly! Should only be called by FilterNodes(). */
+	bool FilterNodesRecursiveNew(
+		const UE::Sequencer::TViewModelPtr<UE::Sequencer::IOutlinerExtension>& InStartNode,
+		TConstArrayView<TNotNull<FSequencerTrackFilter*>> InFilters,
+		const UE::Sequencer::FRecursiveFilterOverrides& InFilterOverrides
+		);
 	
 	/** Invoked when the text filter changes. */
 	void HandleTextFilterChanged();
diff --git a/Engine/Source/Editor/Sequencer/Private/Filters/SequencerFilterData.cpp b/Engine/Source/Editor/Sequencer/Private/Filters/SequencerFilterData.cpp
--- a/Engine/Source/Editor/Sequencer/Private/Filters/SequencerFilterData.cpp
+++ b/Engine/Source/Editor/Sequencer/Private/Filters/SequencerFilterData.cpp
@@ -63,6 +63,11 @@ void FSequencerFilterData::IncrementTotalNodeCount()
 	++TotalNodeCount;
 }
 
+void FSequencerFilterData::AddTotalNodeCount(int32 InNodeCount)
+{
+	TotalNodeCount += FMath::Max(0, InNodeCount);
+}
+
 void FSequencerFilterData::FilterInNode(TWeakViewModelPtr<IOutlinerExtension> InNodeWeak)
 {
 	FilterInNodes.Add(InNodeWeak);
@@ -87,39 +92,32 @@ void FSequencerFilterData::FilterOutNode(TWeakViewModelPtr<IOutlinerExtension> I
 	}
 }
 
-void FSequencerFilterData::FilterInParentChildNodes(const TViewModelPtr<IOutlinerExtension>& InNode
-	, const bool bInIncludeSelf, const bool bInIncludeParents, const bool bInIncludeChildren)
+void FSequencerFilterData::FilterInParentChildNodes(
+	const TViewModelPtr<IOutlinerExtension>& InNode, const bool bInIncludeSelf, const bool bInIncludeParents, const bool bInIncludeChildren
+	)
 {
-	if (!InNode.IsValid())
-	{
-		return;
-	}
-
-	if (bInIncludeParents)
-	{
-		for (TViewModelPtr<IOutlinerExtension> ParentNode : InNode.AsModel()->GetAncestorsOfType<IOutlinerExtension>())
-		{
-			FilterInNode(ParentNode);
-		}
-	}
+	const ENodeVisitFlags Flags = 
+		(bInIncludeSelf ? ENodeVisitFlags::IncludeSelf : ENodeVisitFlags::None)
+		| (bInIncludeParents ? ENodeVisitFlags::IncludeParents : ENodeVisitFlags::None)
+		| (bInIncludeChildren ? ENodeVisitFlags::IncludeChildren : ENodeVisitFlags::None);
+	FilterInParentChildNodes(InNode, Flags);
+}
 
-	if (bInIncludeSelf)
-	{
-		FilterInNode(InNode);
-	}
+int32 FSequencerFilterData::FilterInParentChildNodes(const TViewModelPtr<IOutlinerExtension>& InNode, ENodeVisitFlags InFlags)
+{
+	const auto FilterInLambda = [this](const TViewModelPtr<IOutlinerExtension>& Node){ FilterInNode(Node); };
+	return ApplyToNodes(InNode, FilterInLambda, InFlags);
+}
 
-	if (bInIncludeChildren)
-	{
-		for (TViewModelPtr<IOutlinerExtension> ChildNode : InNode.AsModel()->GetDescendantsOfType<IOutlinerExtension>())
-		{
-			FilterInNode(ChildNode);
-		}
-	}
+int32 FSequencerFilterData::FilterOutParentChildNodes(const TViewModelPtr<IOutlinerExtension>& InNode, ENodeVisitFlags InFlags)
+{
+	const auto FilterOutLambda = [this](const TViewModelPtr<IOutlinerExtension>& Node){ FilterOutNode(Node); };
+	return ApplyToNodes(InNode, FilterOutLambda, InFlags);
 }
 
 void FSequencerFilterData::FilterInNodeWithAncestors(const TViewModelPtr<IOutlinerExtension>& InNode)
 {
-	FilterInParentChildNodes(InNode, true, true, false);
+	FilterInParentChildNodes(InNode, ENodeVisitFlags::IncludeSelf | ENodeVisitFlags::IncludeParents);
 }
 
 bool FSequencerFilterData::ContainsFilterInNodes(const FSequencerFilterData& InOtherData) const
@@ -225,3 +223,42 @@ UObject* FSequencerFilterData::ResolveTrackBoundObject(ISequencer& InSequencer,
 
 	return nullptr;
 }
+
+int32 FSequencerFilterData::ApplyToNodes(
+	const TViewModelPtr<IOutlinerExtension>& InNode, 
+	TFunctionRef<void(const TViewModelPtr<IOutlinerExtension>&)>&& InCallback,
+	ENodeVisitFlags InFlags
+	)
+{
+	int32 NodeCount = 0;
+	if (!InNode.IsValid())
+	{
+		return NodeCount;
+	}
+	
+	if (EnumHasAnyFlags(InFlags, ENodeVisitFlags::IncludeParents))
+	{
+		for (TViewModelPtr<IOutlinerExtension> ParentNode : InNode.AsModel()->GetAncestorsOfType<IOutlinerExtension>())
+		{
+			++NodeCount;
+			InCallback(ParentNode);
+		}
+	}
+
+	if (EnumHasAnyFlags(InFlags, ENodeVisitFlags::IncludeSelf))
+	{
+		++NodeCount;
+		InCallback(InNode);
+	}
+
+	if (EnumHasAnyFlags(InFlags, ENodeVisitFlags::IncludeChildren))
+	{
+		for (TViewModelPtr<IOutlinerExtension> ChildNode : InNode.AsModel()->GetDescendantsOfType<IOutlinerExtension>())
+		{
+			++NodeCount;
+			InCallback(ChildNode);
+		}
+	}
+	
+	return NodeCount;
+}
\ No newline at end of file
diff --git a/Engine/Source/Editor/Sequencer/Private/Filters/Utils/FilterEvaluation.cpp b/Engine/Source/Editor/Sequencer/Private/Filters/Utils/FilterEvaluation.cpp
new file mode 100644
--- /dev/null
+++ b/Engine/Source/Editor/Sequencer/Private/Filters/Utils/FilterEvaluation.cpp
@@ -0,0 +1,90 @@
+﻿// Copyright Epic Games, Inc. All Rights Reserved.
+
+#include "FilterEvaluation.h"
+
+#include "Filters/FilterResult.h"
+#include "Filters/SequencerTrackFilterBase.h"
+#include "Filters/SequencerTrackFilterCollection.h"
+#include "MVVM/ViewModelPtr.h"
+#include "MVVM/Extensions/IOutlinerExtension.h"
+
+namespace UE::Sequencer
+{
+void AppendIfActive(FSequencerTrackFilter& InFilter, TArray<TNotNull<FSequencerTrackFilter*>>& InOutFiltersToEvaluate)
+{
+	if (InFilter.IsActive())
+	{
+		InOutFiltersToEvaluate.Add(&InFilter);
+	}
+}
+
+void AppendIfActive(const TSharedRef<FSequencerTrackFilterCollection>& InCollection, TArray<TNotNull<FSequencerTrackFilter*>>& InOutFiltersToEvaluate)
+{
+	InCollection->ForEachFilter([&InOutFiltersToEvaluate](const TSharedRef<FSequencerTrackFilter>& InFilter)
+	{
+		AppendIfActive(InFilter, InOutFiltersToEvaluate);
+		return true;
+	});
+}
+
+TOptional<EItemFilterState> FRecursiveFilterOverrides::FindOverride(const FSequencerTrackFilter& InFilter) const
+{
+	for (const FFilterOverride& Override : Overrides)
+	{
+		if (&Override.Filter == &InFilter)
+		{
+			return Override.ResultToApply;
+		}
+	}
+	
+	return {};
+}
+
+void FRecursiveFilterOverrides::AddOverride(const FSequencerTrackFilter& InFilter, EItemFilterState InOverride)
+{
+	for (FFilterOverride& Override : Overrides)
+	{
+		if (&Override.Filter == &InFilter)
+		{
+			Override.ResultToApply = InOverride;
+			return;
+		}
+	}
+	
+	Overrides.Emplace(InFilter, InOverride);
+}
+
+FFilterEvaluation EvaluateFiltersOnItem(
+	const TViewModelPtr<IOutlinerExtension>& InItem, 
+	TConstArrayView<TNotNull<FSequencerTrackFilter*>> InFilters,
+	const FRecursiveFilterOverrides& InOverrides
+	)
+{
+	FFilterEvaluation Result;
+	EItemFilterState& ItemResult = Result.ItemResult;
+	FRecursiveFilterOverrides& NewOverrides = Result.NewRecursiveOverrides;
+	ItemResult = EItemFilterState::Include;
+	NewOverrides = InOverrides;
+	
+	for (const FSequencerTrackFilter* Filter : InFilters)
+	{
+		const TOptional<EItemFilterState> Override = InOverrides.FindOverride(*Filter);
+		const FFilterResult FilterResult = Override 
+			? FFilterResult(*Override, EFilterResultFlags::ApplyResultToChildren) 
+			: Filter->Evaluate(InItem);
+		
+		if (EnumHasAnyFlags(FilterResult.Flags, EFilterResultFlags::ApplyResultToChildren))
+		{
+			NewOverrides.AddOverride(*Filter, FilterResult.ItemState);
+		}
+		
+		ItemResult = Combine(Result.ItemResult, FilterResult.ItemState);
+		if (ItemResult == EItemFilterState::Exclude)
+		{
+			return Result;
+		}
+	}
+	
+	return Result;
+}
+}
diff --git a/Engine/Source/Editor/Sequencer/Private/Filters/Utils/FilterEvaluation.h b/Engine/Source/Editor/Sequencer/Private/Filters/Utils/FilterEvaluation.h
new file mode 100644
--- /dev/null
+++ b/Engine/Source/Editor/Sequencer/Private/Filters/Utils/FilterEvaluation.h
@@ -0,0 +1,65 @@
+﻿// Copyright Epic Games, Inc. All Rights Reserved.
+
+#pragma once 
+
+#include "Filters/FilterResult.h"
+#include "MVVM/Extensions/IOutlinerExtension.h"
+#include "MVVM/ViewModelPtr.h"
+
+class FSequencerTrackFilter;
+class FSequencerTrackFilterCollection;
+
+namespace UE::Sequencer
+{
+enum class EItemFilterState : uint8;
+struct FFilterResult;
+
+void AppendIfActive(FSequencerTrackFilter& InFilter, TArray<TNotNull<FSequencerTrackFilter*>>& InOutFiltersToEvaluate);
+inline void AppendIfActive(const TSharedRef<FSequencerTrackFilter>& InFilter, TArray<TNotNull<FSequencerTrackFilter*>>& InOutFiltersToEvaluate) { AppendIfActive(*InFilter, InOutFiltersToEvaluate); }
+void AppendIfActive(const TSharedRef<FSequencerTrackFilterCollection>& InCollection, TArray<TNotNull<FSequencerTrackFilter*>>& InOutFiltersToEvaluate);
+template <typename TFilterType> requires std::is_base_of_v<FSequencerTrackFilter, TFilterType>
+void AppendIfActive(TConstArrayView<TSharedRef<TFilterType>> InFilters, TArray<TNotNull<FSequencerTrackFilter*>>& InOutFiltersToEvaluate);
+
+struct FFilterOverride
+{
+	const FSequencerTrackFilter& Filter;
+	EItemFilterState ResultToApply;
+
+	explicit FFilterOverride(const FSequencerTrackFilter& Filter UE_LIFETIMEBOUND, EItemFilterState ResultToApply)
+		: Filter(Filter)
+		, ResultToApply(ResultToApply)
+	{}
+};
+
+struct FRecursiveFilterOverrides
+{
+	TArray<FFilterOverride, TInlineAllocator<8>> Overrides;
+	
+	TOptional<EItemFilterState> FindOverride(const FSequencerTrackFilter& InFilter) const;
+	void AddOverride(const FSequencerTrackFilter& InFilter, EItemFilterState InOverride);
+};
+
+struct FFilterEvaluation
+{
+	EItemFilterState ItemResult;
+	FRecursiveFilterOverrides NewRecursiveOverrides;
+};
+
+FFilterEvaluation EvaluateFiltersOnItem(
+	const TViewModelPtr<IOutlinerExtension>& InItem, 
+	TConstArrayView<TNotNull<FSequencerTrackFilter*>> InFilters,
+	const FRecursiveFilterOverrides& InOverrides
+	);
+}
+
+namespace UE::Sequencer
+{
+template <typename TFilterType> requires std::is_base_of_v<FSequencerTrackFilter, TFilterType>
+void AppendIfActive(TConstArrayView<TSharedRef<TFilterType>> InFilters, TArray<TNotNull<FSequencerTrackFilter*>>& InOutFiltersToEvaluate)
+{
+	for (const TSharedRef<TFilterType>& Filter : InFilters)
+	{
+		AppendIfActive(*Filter, InOutFiltersToEvaluate);
+	}
+}
+}
diff --git a/Engine/Source/Editor/Sequencer/Public/Filters/FilterResult.h b/Engine/Source/Editor/Sequencer/Public/Filters/FilterResult.h
new file mode 100644
--- /dev/null
+++ b/Engine/Source/Editor/Sequencer/Public/Filters/FilterResult.h
@@ -0,0 +1,69 @@
+﻿// Copyright Epic Games, Inc. All Rights Reserved.
+
+#pragma once
+
+#include "HAL/Platform.h"
+#include "Misc/AssertionMacros.h"
+#include "Misc/EnumClassFlags.h"
+
+namespace UE::Sequencer
+{
+/** Describes what a filter wants to happen to an item. */
+enum class EItemFilterState : uint8
+{
+	/** The item should be displayed to the user. */
+	Include,
+	/** The item should not be displayed to the user. */
+	Exclude,
+	/** The filter does not care. Let other filters decide. */
+	DoNotCare
+};
+
+/** @return Include if either is set to Include and the other to Include or DoNotCare, Exclude if either is set to exclude, and DoNotCare otherwise. */
+EItemFilterState Combine(EItemFilterState InLeft, EItemFilterState InRight);
+
+/** Additional info filters may want to return after having evaluated an item. */
+enum class EFilterResultFlags : uint8
+{
+	None,
+	
+	/** 
+	 * Apply EItemFilterState to all children. 
+	 * Example: 
+	 * - Suppose EItemFilterState::Include is returned on "Cube.Transform" 
+	 * - Also assume EItemFilterState::Include for "Cube.Transform.Location", "Cube.Transform.Location.X", etc. 
+	 * This allows filter optimization to optimize needless evaluation.
+	 */
+	ApplyResultToChildren = 1 << 0
+};
+ENUM_CLASS_FLAGS(EFilterResultFlags);
+
+/** The result of evaluating a filter. */
+struct FFilterResult
+{
+	/** What should happen to the item state */
+	EItemFilterState ItemState;
+	
+	/** Additional information the filtered returned, e.g. that the filter need not be invoked for children. */
+	EFilterResultFlags Flags;
+
+	explicit FFilterResult(EItemFilterState ItemState, EFilterResultFlags Flags = EFilterResultFlags::None)
+		: ItemState(ItemState)
+		, Flags(Flags)
+	{}
+};
+}
+
+namespace UE::Sequencer
+{
+inline EItemFilterState Combine(EItemFilterState InLeft, EItemFilterState InRight)
+{
+	switch (InLeft)
+	{
+	case EItemFilterState::Include: return InRight == EItemFilterState::Exclude ? EItemFilterState::Exclude : InLeft;
+	case EItemFilterState::Exclude: return EItemFilterState::Exclude;
+	case EItemFilterState::DoNotCare: return InRight;
+	default: checkNoEntry(); return EItemFilterState::DoNotCare;
+	}
+}
+}
\ No newline at end of file
diff --git a/Engine/Source/Editor/Sequencer/Public/Filters/Filters/SequencerTrackFilter_Selected.h b/Engine/Source/Editor/Sequencer/Public/Filters/Filters/SequencerTrackFilter_Selected.h
--- a/Engine/Source/Editor/Sequencer/Public/Filters/Filters/SequencerTrackFilter_Selected.h
+++ b/Engine/Source/Editor/Sequencer/Public/Filters/Filters/SequencerTrackFilter_Selected.h
@@ -19,6 +19,7 @@ public:
 	virtual FText GetDefaultToolTipText() const override;
 	virtual TSharedPtr<FUICommandInfo> GetToggleCommand() const override;
 	virtual void ActiveStateChanged(const bool bInActive) override;
+	virtual UE::Sequencer::FFilterResult Evaluate(FSequencerTrackFilterType InItem) const override;
 	//~ End FSequencerTrackFilter
 
 	//~ Begin FFilterBase
@@ -26,23 +27,34 @@ public:
 	virtual FSlateIcon GetIcon() const override;
 	//~ End FFilterBase
 
+	virtual void OnPreFilter() override;
+
 	//~ Begin IFilter
 	virtual FString GetName() const override;
 	virtual bool PassesFilter(FSequencerTrackFilterType InItem) const override;
 	//~ End IFilter
 
-	void ToggleShowOnlySelectedTracks();
-
 protected:
+	
 	void BindSelectionChanged();
 	void UnbindSelectionChanged();
 
-	void OnSelectionChanged(UObject* const InObject);
-
 	FDelegateHandle OnSelectionChangedHandle;
 	
 private:
 	
 	/** Objects extending this filter. Cached at construction to speed up filtering. */
 	const TArray<TWeakObjectPtr<USequencerTrackFilterSelectedExtension>> CachedWeakExtensions;
+	
+	/** Whether IncludedViewModels needs to be rebuilt. */
+	bool bNeedsPrefilterRebuild = true;
+	
+	/** The view models that pass the filter. Built once during pre-filter to avoid excessive iteration of tree hierarchy. */
+	TMap<UE::Sequencer::TWeakViewModelPtr<UE::Sequencer::FViewModel>, UE::Sequencer::EFilterResultFlags> IncludedViewModels;
+	
+	void BindEvents();
+	void UnbindEvents();
+	
+	void OnSequencerSelectionChanged();
+	void OnViewportSelectionChanged(UObject* const InObject);
 };
diff --git a/Engine/Source/Editor/Sequencer/Public/Filters/SequencerFilterData.h b/Engine/Source/Editor/Sequencer/Public/Filters/SequencerFilterData.h
--- a/Engine/Source/Editor/Sequencer/Public/Filters/SequencerFilterData.h
+++ b/Engine/Source/Editor/Sequencer/Public/Filters/SequencerFilterData.h
@@ -3,8 +3,9 @@
 #pragma once
 
 #include "FilterBarDelegates.h"
-#include "MovieSceneTrack.h"
 #include "MVVM/ViewModelPtr.h"
+#include "Misc/EnumClassFlags.h"
+#include "MovieSceneTrack.h"
 
 #define UE_API SEQUENCER_API
 
@@ -19,6 +20,22 @@ namespace UE::Sequencer
 
 using FSequencerTrackFilterType = UE::Sequencer::FViewModelPtr;
 
+namespace UE::Sequencer
+{
+/** Utility argument for FSequencerFilterData. */
+enum class ENodeVisitFlags : uint8
+{
+	None,
+	
+	IncludeSelf = 1 << 0,
+	IncludeParents = 1 << 1,
+	IncludeChildren = 1 << 2,
+	
+	IncludeAll = IncludeSelf | IncludeParents | IncludeChildren,
+};
+ENUM_CLASS_FLAGS(ENodeVisitFlags);
+}
+
 /** Represents a cache between nodes for a filter operation. */
 struct FSequencerFilterData
 {
@@ -38,12 +55,23 @@ struct FSequencerFilterData
 	UE_API uint32 GetFilterOutCount() const;
 
 	UE_API void IncrementTotalNodeCount();
+	UE_API void AddTotalNodeCount(int32 InNodeCount);
 
 	UE_API void FilterInNode(UE::Sequencer::TWeakViewModelPtr<UE::Sequencer::IOutlinerExtension> InNodeWeak);
 	UE_API void FilterOutNode(UE::Sequencer::TWeakViewModelPtr<UE::Sequencer::IOutlinerExtension> InNodeWeak);
 
-	UE_API void FilterInParentChildNodes(const UE::Sequencer::TViewModelPtr<UE::Sequencer::IOutlinerExtension>& InNode
-		, const bool bInIncludeSelf, const bool bInIncludeParents, const bool bInIncludeChildren = false);
+	UE_DEPRECATED(5.8, "Use the version with EFilterInFlags instead")
+	UE_API void FilterInParentChildNodes(const UE::Sequencer::TViewModelPtr<UE::Sequencer::IOutlinerExtension>& InNode, const bool bInIncludeSelf, const bool bInIncludeParents, const bool bInIncludeChildren = false);
+	
+	/** @return Number of nodes visited */
+	UE_API int32 FilterInParentChildNodes(
+		const UE::Sequencer::TViewModelPtr<UE::Sequencer::IOutlinerExtension>& InNode, UE::Sequencer::ENodeVisitFlags InFlags
+		);
+	/** @return Number of nodes visited */
+	UE_API int32 FilterOutParentChildNodes(
+		const UE::Sequencer::TViewModelPtr<UE::Sequencer::IOutlinerExtension>& InNode, 
+		UE::Sequencer::ENodeVisitFlags InFlags
+		);
 
 	UE_API void FilterInNodeWithAncestors(const UE::Sequencer::TViewModelPtr<UE::Sequencer::IOutlinerExtension>& InNode);
 
@@ -72,6 +100,13 @@ protected:
 	
 	/** Broadcasts when the filter state of IOutlinerExtension is changed. */
 	UE::Sequencer::FChangeOutlinerExtensionFilterState ChangeFilterStateEvent;
+	
+	/** @return Number of nodes visited */
+	int32 ApplyToNodes(
+		const UE::Sequencer::TViewModelPtr<UE::Sequencer::IOutlinerExtension>& InNode, 
+		TFunctionRef<void(const UE::Sequencer::TViewModelPtr<UE::Sequencer::IOutlinerExtension>&)>&& InCallback,
+		UE::Sequencer::ENodeVisitFlags InFlags
+		);
 };
 
 #undef UE_API
diff --git a/Engine/Source/Editor/Sequencer/Public/Filters/SequencerTrackFilterBase.h b/Engine/Source/Editor/Sequencer/Public/Filters/SequencerTrackFilterBase.h
--- a/Engine/Source/Editor/Sequencer/Public/Filters/SequencerTrackFilterBase.h
+++ b/Engine/Source/Editor/Sequencer/Public/Filters/SequencerTrackFilterBase.h
@@ -2,6 +2,7 @@
 
 #pragma once
 
+#include "FilterResult.h"
 #include "Filters/FilterBase.h"
 #include "Filters/ISequencerTrackFilters.h"
 #include "GameFramework/Actor.h"
@@ -49,11 +50,12 @@ public:
 	SEQUENCER_API FSequencerTrackFilter(ISequencerTrackFilters& InOutFilterInterface, TSharedPtr<FFilterCategory>&& InCategory = nullptr);
 
 	//~ Begin FSequencerFilterBase
-
 	SEQUENCER_API virtual ISequencerTrackFilters& GetFilterInterface() const override;
-
 	//~ End FSequencerFilterBase
 
+	/** Evaluates the filter on an item. It allows returning more information to optimize evaluation. */
+	SEQUENCER_API virtual UE::Sequencer::FFilterResult Evaluate(FSequencerTrackFilterType InItem) const;
+	
 	/** Returns whether this filter needs reevaluating any time track values have been modified, not just tree changes */
 	virtual bool ShouldUpdateOnTrackValueChanged() const { return false; }
 
