Sequencer: Replace unacceptable word in code comment.

#Synced-CL 51881412

diff --git a/Engine/Source/Editor/Sequencer/Private/Sequencer.cpp b/Engine/Source/Editor/Sequencer/Private/Sequencer.cpp
--- a/Engine/Source/Editor/Sequencer/Private/Sequencer.cpp
+++ b/Engine/Source/Editor/Sequencer/Private/Sequencer.cpp
@@ -4412,7 +4412,7 @@ void FSequencer::DestroySequencerWidget()
 		SequencerWrapperWidget.Reset();
 	}
 	
-	// Last reference to SSequencer should have been nuked now.
+	// Last reference to SSequencer should have been cleared now.
 	ensure(!WeakSequencerWidget.IsValid());
 }
 
