Sequencer: Disable new filtering algorithm due to bugs. See UE-372023.

#Synced-CL 51881412

diff --git a/Engine/Source/Editor/Sequencer/Private/Filters/SequencerFilterBar.cpp b/Engine/Source/Editor/Sequencer/Private/Filters/SequencerFilterBar.cpp
--- a/Engine/Source/Editor/Sequencer/Private/Filters/SequencerFilterBar.cpp
+++ b/Engine/Source/Editor/Sequencer/Private/Filters/SequencerFilterBar.cpp
@@ -1420,7 +1420,7 @@ void FSequencerFilterBar::PreFilter() const
 }
 
 // TODO UE-371829: Remove this CVar: Also look for "Sequencer.UseNewFilteringAlgorithm" because it passed to FindConsoleVariable
-static TAutoConsoleVariable<bool> CVarUseNewFilteringAlgorithm(TEXT("Sequencer.UseNewFilteringAlgorithm"), true, TEXT(""));
+static TAutoConsoleVariable<bool> CVarUseNewFilteringAlgorithm(TEXT("Sequencer.UseNewFilteringAlgorithm"), false, TEXT(""));
 
 const FSequencerFilterData& FSequencerFilterBar::FilterNodes()
 {
