Sequencer: Enable root motion extraction for control rig in Animation Mixer.

The bulk of this change (and the related other changes around it) are to enable decorations with channels to live on tracks via a 'host section' that the decoration owns. This involved allowing keying paths for channels owned directly by those decorations by expanding the role of the decoration editors. It also involved some extra view model side plumbing to ensure updates work as expected when modifying channels on the decorations. I also plugged in the root motion inverse to the control rig controls when root motion is being applied so that the controls don't double offset when they affect root motion.

#Synced-CL 52256411

diff --git a/Engine/Plugins/Animation/ControlRig/Source/ControlRigEditor/Private/EditMode/ControlRigEditMode.cpp b/Engine/Plugins/Animation/ControlRig/Source/ControlRigEditor/Private/EditMode/ControlRigEditMode.cpp
--- a/Engine/Plugins/Animation/ControlRig/Source/ControlRigEditor/Private/EditMode/ControlRigEditMode.cpp
+++ b/Engine/Plugins/Animation/ControlRig/Source/ControlRigEditor/Private/EditMode/ControlRigEditMode.cpp
@@ -93,6 +93,8 @@
 #include "Tools/ControlRigFrustumSelectInteraction.h"
 #include "Tools/MotionTrailOptions.h"
 #include "Transform/TransformableHandleUtils.h"
+#include "Features/IModularFeatures.h"
+#include "IMovieSceneRootMotionOffsetProvider.h"
 #include "ViewportInteractions/ViewportInteractionsBehaviorSource.h"
 
 #include UE_INLINE_GENERATED_CPP_BY_NAME(ControlRigEditMode)
@@ -6499,8 +6501,39 @@ FTransform FControlRigEditMode::GetHostingSceneComponentTransform(const UControl
 			ControlRig = GetControlRigs()[0].Get();
 		}
 
-		const USceneComponent* HostingComponent = GetHostingSceneComponent(ControlRig);
-		return HostingComponent ? HostingComponent->GetComponentTransform() : FTransform::Identity;
+		USceneComponent* HostingComponent = GetHostingSceneComponent(ControlRig);
+		FTransform Result = HostingComponent ? HostingComponent->GetComponentTransform() : FTransform::Identity;
+
+		// When root motion is extracted from a control rig layer and applied to the actor,
+		// the control hierarchy's cached global transforms still include the extracted root
+		// motion while the component transform also includes it (via the moved actor).
+		// Subtract the root motion offset so gizmos are not drawn at a double offset.
+		if (TSharedPtr<ISequencer> Sequencer = WeakSequencer.Pin())
+		{
+			const UMovieSceneEntitySystemLinker* EntityLinker =
+				Sequencer->GetEvaluationTemplate().GetEntitySystemLinker();
+			if (EntityLinker && HostingComponent)
+			{
+				TArray<IMovieSceneRootMotionOffsetProvider*> Providers =
+					IModularFeatures::Get().GetModularFeatureImplementations<IMovieSceneRootMotionOffsetProvider>(
+						IMovieSceneRootMotionOffsetProvider::GetModularFeatureName());
+				for (IMovieSceneRootMotionOffsetProvider* Provider : Providers)
+				{
+					if (!Provider)
+					{
+						continue;
+					}
+					FTransform Offset = Provider->GetRootMotionOffset(EntityLinker, HostingComponent);
+					if (!Offset.Equals(FTransform::Identity))
+					{
+						Result = Offset.Inverse() * Result;
+						break;
+					}
+				}
+			}
+		}
+
+		return Result;
 	}
 	return FTransform::Identity;
 }
diff --git a/Engine/Plugins/Experimental/MovieSceneAnimMixer/Source/MovieSceneAnimMixer/Private/MovieSceneAnimationMixerTrack.cpp b/Engine/Plugins/Experimental/MovieSceneAnimMixer/Source/MovieSceneAnimMixer/Private/MovieSceneAnimationMixerTrack.cpp
--- a/Engine/Plugins/Experimental/MovieSceneAnimMixer/Source/MovieSceneAnimMixer/Private/MovieSceneAnimationMixerTrack.cpp
+++ b/Engine/Plugins/Experimental/MovieSceneAnimMixer/Source/MovieSceneAnimMixer/Private/MovieSceneAnimationMixerTrack.cpp
@@ -37,6 +37,24 @@ bool FMovieSceneByteChannelDefaultOnly::SerializeFromMismatchedTag(const FProper
 	return false;
 }
 
+void UMovieSceneRootMotionHostSection::SetOwnerDecoration(UObject* InOwnerDecoration)
+{
+	OwnerDecoration = InOwnerDecoration;
+}
+
+EMovieSceneChannelProxyType UMovieSceneRootMotionHostSection::CacheChannelProxy()
+{
+	if (IMovieSceneChannelDecoration* ChannelDecoration = Cast<IMovieSceneChannelDecoration>(OwnerDecoration))
+	{
+		FMovieSceneChannelProxyData ChannelProxyData;
+		EMovieSceneChannelProxyType ProxyType = ChannelDecoration->PopulateChannelProxy(ChannelProxyData);
+		ChannelProxy = MakeShared<FMovieSceneChannelProxy>(MoveTemp(ChannelProxyData));
+		return ProxyType;
+	}
+	ChannelProxy = MakeShared<FMovieSceneChannelProxy>();
+	return EMovieSceneChannelProxyType::Dynamic;
+}
+
 UMovieSceneRootMotionSettingsDecoration::UMovieSceneRootMotionSettingsDecoration(const FObjectInitializer& ObjInit)
 	: Super(ObjInit)
 {
@@ -221,7 +239,6 @@ EMovieSceneChannelProxyType UMovieSceneRootMotionSettingsDecoration::PopulateCha
 void UMovieSceneRootMotionSettingsDecoration::ExtendEntityImpl(UMovieSceneEntitySystemLinker* EntityLinker, const FEntityImportParams& Params, FImportedEntity* OutImportedEntity)
 {
 	using namespace UE::MovieScene;
-	FBuiltInComponentTypes* BuiltInComponents = FBuiltInComponentTypes::Get();
 	FAnimMixerComponentTypes* AnimMixerComponents = FAnimMixerComponentTypes::Get();
 
 	FMovieSceneRootMotionSettings RootMotionSettings;
@@ -239,20 +256,47 @@ void UMovieSceneRootMotionSettingsDecoration::ExtendEntityImpl(UMovieSceneEntity
 
 	RootMotionSettings.RootOriginLocation = RootOriginLocation;
 
-	const bool bAddTransform = RootMotionSettings.TransformMode != EMovieSceneRootMotionTransformMode::Asset;
+	if (RootMotionSettings.TransformMode != EMovieSceneRootMotionTransformMode::Asset)
+	{
+		RootMotionSettings.OffsetChannels[0] = &Location[0];
+		RootMotionSettings.OffsetChannels[1] = &Location[1];
+		RootMotionSettings.OffsetChannels[2] = &Location[2];
+		RootMotionSettings.OffsetChannels[3] = &Rotation[0];
+		RootMotionSettings.OffsetChannels[4] = &Rotation[1];
+		RootMotionSettings.OffsetChannels[5] = &Rotation[2];
+	}
+
 	OutImportedEntity->AddBuilder(
 		FEntityBuilder()
-		.AddConditional(BuiltInComponents->DoubleChannel[0], &Location[0], bAddTransform)
-		.AddConditional(BuiltInComponents->DoubleChannel[1], &Location[1], bAddTransform)
-		.AddConditional(BuiltInComponents->DoubleChannel[2], &Location[2], bAddTransform)
-		.AddConditional(BuiltInComponents->DoubleChannel[3], &Rotation[0], bAddTransform)
-		.AddConditional(BuiltInComponents->DoubleChannel[4], &Rotation[1], bAddTransform)
-		.AddConditional(BuiltInComponents->DoubleChannel[5], &Rotation[2], bAddTransform)
 		.Add(AnimMixerComponents->RootMotionSettings, RootMotionSettings)
 		.AddConditional(AnimMixerComponents->BoneMatchTransform, GetBoneMatchData().MatchTransform, HasBoneMatch() && GetBoneMatchData().bIsValid)
 	);
 }
 
+TArrayView<TObjectPtr<UMovieSceneSection>> UMovieSceneRootMotionSettingsDecoration::GetSections()
+{
+	// Only provide a host section when the decoration is on a track (not a section).
+	// When on a section, return empty so the channel decoration path handles it instead.
+	if (!GetTypedOuter<UMovieSceneSection>())
+	{
+		EnsureHostSection();
+		return TArrayView<TObjectPtr<UMovieSceneSection>>(&HostSection, 1);
+	}
+	return TArrayView<TObjectPtr<UMovieSceneSection>>();
+}
+
+void UMovieSceneRootMotionSettingsDecoration::EnsureHostSection()
+{
+	if (!HostSection)
+	{
+		UMovieSceneRootMotionHostSection* NewSection = NewObject<UMovieSceneRootMotionHostSection>(this, NAME_None, RF_Transactional);
+		NewSection->SetOwnerDecoration(this);
+		NewSection->SetRange(TRange<FFrameNumber>::All());
+		NewSection->SetColorTint(FColor(20, 70, 70, 200));
+		HostSection = NewSection;
+	}
+}
+
 void UMovieSceneAnimationSectionDecoration::ExtendEntityImpl(UMovieSceneEntitySystemLinker* EntityLinker, const FEntityImportParams& Params, FImportedEntity* OutImportedEntity)
 {
 	using namespace UE::MovieScene;
diff --git a/Engine/Plugins/Experimental/MovieSceneAnimMixer/Source/MovieSceneAnimMixer/Private/Systems/MovieSceneAnimMixerSystem.cpp b/Engine/Plugins/Experimental/MovieSceneAnimMixer/Source/MovieSceneAnimMixer/Private/Systems/MovieSceneAnimMixerSystem.cpp
--- a/Engine/Plugins/Experimental/MovieSceneAnimMixer/Source/MovieSceneAnimMixer/Private/Systems/MovieSceneAnimMixerSystem.cpp
+++ b/Engine/Plugins/Experimental/MovieSceneAnimMixer/Source/MovieSceneAnimMixer/Private/Systems/MovieSceneAnimMixerSystem.cpp
@@ -5,6 +5,7 @@
 #include "MovieSceneAnimationMixerTrack.h"
 #include "MovieSceneRootMotionSection.h"
 #include "MovieSceneRootMotionTargetDecoration.h"
+#include "Channels/MovieSceneDoubleChannel.h"
 #include "Systems/ByteChannelEvaluatorSystem.h"
 #include "Systems/MovieSceneMixerBakeTargetSystem.h"
 
@@ -620,12 +621,7 @@ namespace UE::MovieScene
 			const FMovieSceneRootMotionSettings* RootMotionSettings,
 			const double* WeightAndEasings,
 			const FTransform* BoneMatchTransform,
-			const double* OffsetX,
-			const double* OffsetY,
-			const double* OffsetZ,
-			const double* RotOffsetX,
-			const double* RotOffsetY,
-			const double* RotOffsetZ)
+			const FFrameTime* OptEvalTime)
 		{
 			if (MixerEntry)
 			{
@@ -638,17 +634,24 @@ namespace UE::MovieScene
 				{
 					MixerEntry->RootMotionSettings = *RootMotionSettings;
 
-					// Apply evaluated offset channel results to the root motion settings
-					FVector Location = RootMotionSettings->RootLocation;
-					FRotator Rotation = RootMotionSettings->RootRotation.Rotator();
-					if (OffsetX)    { Location.X = *OffsetX; }
-					if (OffsetY)    { Location.Y = *OffsetY; }
-					if (OffsetZ)    { Location.Z = *OffsetZ; }
-					if (RotOffsetX) { Rotation.Roll  = *RotOffsetX; }
-					if (RotOffsetY) { Rotation.Pitch = *RotOffsetY; }
-					if (RotOffsetZ) { Rotation.Yaw   = *RotOffsetZ; }
-					MixerEntry->RootMotionSettings->RootLocation = Location;
-					MixerEntry->RootMotionSettings->RootRotation = Rotation.Quaternion();
+					// Evaluate offset channels from the decoration into RootLocation/RootRotation.
+					// The channels are stored as pointers in the settings to avoid using
+					// BuiltInComponents::DoubleChannel which clashes with control rig parameters.
+					if (OptEvalTime)
+					{
+						FMovieSceneRootMotionSettings& Settings = *MixerEntry->RootMotionSettings;
+						FVector Location = Settings.RootLocation;
+						FRotator Rotation = Settings.RootRotation.Rotator();
+						const FMovieSceneDoubleChannel* const* Channels = Settings.OffsetChannels;
+						if (Channels[0]) { double Val; if (Channels[0]->Evaluate(*OptEvalTime, Val)) { Location.X = Val; } }
+						if (Channels[1]) { double Val; if (Channels[1]->Evaluate(*OptEvalTime, Val)) { Location.Y = Val; } }
+						if (Channels[2]) { double Val; if (Channels[2]->Evaluate(*OptEvalTime, Val)) { Location.Z = Val; } }
+						if (Channels[3]) { double Val; if (Channels[3]->Evaluate(*OptEvalTime, Val)) { Rotation.Roll  = Val; } }
+						if (Channels[4]) { double Val; if (Channels[4]->Evaluate(*OptEvalTime, Val)) { Rotation.Pitch = Val; } }
+						if (Channels[5]) { double Val; if (Channels[5]->Evaluate(*OptEvalTime, Val)) { Rotation.Yaw   = Val; } }
+						Settings.RootLocation = Location;
+						Settings.RootRotation = Rotation.Quaternion();
+					}
 				}
 
 				if (BoneMatchTransform)
@@ -1585,12 +1588,7 @@ void UMovieSceneAnimMixerSystem::OnSchedulePersistentTasks(UE::MovieScene::IEnti
 		.ReadOptional(AnimMixerComponents->RootMotionSettings)
 		.ReadOptional(BuiltInComponents->WeightAndEasingResult)
 		.ReadOptional(AnimMixerComponents->BoneMatchTransform)
-		.ReadOptional(BuiltInComponents->DoubleResult[0])
-		.ReadOptional(BuiltInComponents->DoubleResult[1])
-		.ReadOptional(BuiltInComponents->DoubleResult[2])
-		.ReadOptional(BuiltInComponents->DoubleResult[3])
-		.ReadOptional(BuiltInComponents->DoubleResult[4])
-		.ReadOptional(BuiltInComponents->DoubleResult[5])
+		.ReadOptional(BuiltInComponents->EvalTime)
 		.FilterNone({BuiltInComponents->Tags.NeedsUnlink, BuiltInComponents->Tags.Ignored})
 		.Schedule_PerEntity<FUpdateTaskEntities>(&Linker->EntityManager, TaskScheduler);
 
diff --git a/Engine/Plugins/Experimental/MovieSceneAnimMixer/Source/MovieSceneAnimMixer/Public/MovieSceneAnimationMixerTrack.h b/Engine/Plugins/Experimental/MovieSceneAnimMixer/Source/MovieSceneAnimMixer/Public/MovieSceneAnimationMixerTrack.h
--- a/Engine/Plugins/Experimental/MovieSceneAnimMixer/Source/MovieSceneAnimMixer/Public/MovieSceneAnimationMixerTrack.h
+++ b/Engine/Plugins/Experimental/MovieSceneAnimMixer/Source/MovieSceneAnimMixer/Public/MovieSceneAnimationMixerTrack.h
@@ -7,6 +7,7 @@
 #include "Channels/MovieSceneByteChannel.h"
 #include "Decorations/IMovieSceneChannelDecoration.h"
 #include "Decorations/IMovieSceneSectionDecoration.h"
+#include "Decorations/IMovieSceneSectionProviderDecoration.h"
 #include "EntitySystem/IMovieSceneEntityDecorator.h"
 #include "EntitySystem/MovieSceneEntityIDs.h"
 #include "MovieSceneMixedAnimationTarget.h"
@@ -86,16 +87,39 @@ private:
 	FColor MixerTintOverride = FColor();
 };
 
+// Section that hosts channels from a root motion settings decoration when it lives
+// on a track rather than a section. Delegates CacheChannelProxy to the decoration.
+UCLASS(MinimalAPI)
+class UMovieSceneRootMotionHostSection : public UMovieSceneSection
+{
+	GENERATED_BODY()
+
+public:
+	UMovieSceneRootMotionHostSection()
+	{
+		bSupportsInfiniteRange = true;
+	}
+
+	MOVIESCENEANIMMIXER_API virtual EMovieSceneChannelProxyType CacheChannelProxy() override;
+
+	MOVIESCENEANIMMIXER_API void SetOwnerDecoration(UObject* InOwnerDecoration);
+
+	UPROPERTY()
+	TObjectPtr<UObject> OwnerDecoration;
+};
+
 /**
- * Section-level decoration that provides root motion offset settings.
- * Adds transform offset channels to an existing animation section.
- * Channels appear under "Modifiers" in the Sequencer outliner.
+ * Decoration that provides root motion offset settings.
+ * Can live on a section (skeletal animation) or on a track (control rig child track).
+ * When on a section, channels are hosted by the parent section via IMovieSceneChannelDecoration.
+ * When on a track, an internal host section is created via IMovieSceneSectionProviderDecoration.
  */
 UCLASS(MinimalAPI, collapseCategories, meta=(DisplayName="Root Motion Settings"))
 class UMovieSceneRootMotionSettingsDecoration
 	: public UMovieSceneSignedObject
 	, public IMovieSceneChannelDecoration
 	, public IMovieSceneEntityDecorator
+	, public IMovieSceneSectionProviderDecoration
 {
 public:
 
@@ -109,6 +133,12 @@ public:
 	virtual EMovieSceneChannelProxyType PopulateChannelProxy(FMovieSceneChannelProxyData& OutProxyData) override;
 	virtual void ExtendEntityImpl(UMovieSceneEntitySystemLinker* EntityLinker, const FEntityImportParams& Params, FImportedEntity* OutImportedEntity) override;
 
+	// IMovieSceneSectionProviderDecoration - provides an internal section when on a track
+	virtual TArrayView<TObjectPtr<UMovieSceneSection>> GetSections() override;
+
+	// Creates the internal host section when the decoration is added to a track
+	MOVIESCENEANIMMIXER_API void EnsureHostSection();
+
 	// Transform Offset Channels
 
 	UPROPERTY()
@@ -167,6 +197,10 @@ public:
 	UPROPERTY()
 	bool bInitialRootTransformValid = false;
 
+	// Internal section for hosting channels when the decoration lives on a track
+	UPROPERTY()
+	TObjectPtr<UMovieSceneSection> HostSection;
+
 #if WITH_EDITORONLY_DATA
 	// When true, the edit mode draws this section's skeleton in the viewport. Useful for bone matching.
 	UPROPERTY(Transient)
diff --git a/Engine/Plugins/Experimental/MovieSceneAnimMixer/Source/MovieSceneAnimMixer/Public/Systems/MovieSceneRootMotionSystem.h b/Engine/Plugins/Experimental/MovieSceneAnimMixer/Source/MovieSceneAnimMixer/Public/Systems/MovieSceneRootMotionSystem.h
--- a/Engine/Plugins/Experimental/MovieSceneAnimMixer/Source/MovieSceneAnimMixer/Public/Systems/MovieSceneRootMotionSystem.h
+++ b/Engine/Plugins/Experimental/MovieSceneAnimMixer/Source/MovieSceneAnimMixer/Public/Systems/MovieSceneRootMotionSystem.h
@@ -13,6 +13,7 @@ enum class EMovieSceneRootMotionDestination : uint8;
 enum class EMovieSceneRootMotionGapBehavior : uint8;
 
 class UMovieSceneAnimMixerSystem;
+struct FMovieSceneDoubleChannel;
 class UMovieSceneRootMotionSection;
 class UMovieSceneRootMotionSystem;
 
@@ -97,6 +98,11 @@ struct FMovieSceneRootMotionSettings
 	/** If true, the animation has Force Root Lock + Enable Root Motion. The extraction task
 	 *  handles root locking, and the eval task suppresses the lock during pose evaluation. */
 	uint8 bForceRootLock : 1 = 0;
+
+	// Offset channels from the root motion settings decoration.
+	// Evaluated per-frame by the consumer to populate RootLocation/RootRotation.
+	// [0..2] = Location X/Y/Z, [3..5] = Rotation Roll/Pitch/Yaw.
+	const FMovieSceneDoubleChannel* OffsetChannels[6] = {};
 };
 
 // Structure that is shared between entities for handling the mixer's root motion result. 
diff --git a/Engine/Plugins/Experimental/MovieSceneAnimMixer/Source/MovieSceneAnimMixerEditor/Private/MVVM/AnimationMixerLayerModel.cpp b/Engine/Plugins/Experimental/MovieSceneAnimMixer/Source/MovieSceneAnimMixerEditor/Private/MVVM/AnimationMixerLayerModel.cpp
--- a/Engine/Plugins/Experimental/MovieSceneAnimMixer/Source/MovieSceneAnimMixerEditor/Private/MVVM/AnimationMixerLayerModel.cpp
+++ b/Engine/Plugins/Experimental/MovieSceneAnimMixer/Source/MovieSceneAnimMixerEditor/Private/MVVM/AnimationMixerLayerModel.cpp
@@ -34,6 +34,7 @@
 #include "MVVM/DecorationModelStorageExtension.h"
 #include "SequencerCommonHelpers.h"
 #include "Decorations/MovieSceneTrackRowDecoration.h"
+#include "MVVM/ViewModels/OutlinerColumns/OutlinerColumnTypes.h"
 
 #define LOCTEXT_NAMESPACE "AnimationMixerLayerModel"
 
@@ -128,6 +129,39 @@ FAnimationMixerLayerModel::FAnimationMixerLayerModel(UMovieSceneAnimationMixerTr
 
 FAnimationMixerLayerModel::~FAnimationMixerLayerModel()
 {
+	if (ChildTrackChangedHandle.IsValid() && WeakChildTrack.IsValid())
+	{
+		WeakChildTrack.Get()->OnSignatureChanged().Remove(ChildTrackChangedHandle);
+	}
+}
+
+void FAnimationMixerLayerModel::OnChildTrackChanged()
+{
+	UMovieSceneTrack* ChildTrack = WeakChildTrack.Get();
+	if (!ChildTrack)
+	{
+		return;
+	}
+
+	TSharedPtr<FSequenceModel> SequenceModel = FindAncestorOfType<FSequenceModel>();
+	if (!SequenceModel)
+	{
+		return;
+	}
+
+	if (FDecorationModelStorageExtension* DecorationModelStorage = SequenceModel->CastDynamic<FDecorationModelStorageExtension>())
+	{
+		UMovieSceneAnimationMixerTrack* ParentTrack = WeakParentTrack.Get();
+		if (ParentTrack)
+		{
+			FViewModelChildren OutlinerChildren = GetChildList(EViewModelListType::Outliner);
+			DecorationModelStorage->SyncDecorationModels(ChildTrack, ParentTrack, OutlinerChildren);
+		}
+	}
+
+	// The child track's signature changes on undo/redo, property edits, etc.
+	// Refresh the layout so section models with stale channel proxies get rebuilt.
+	RefreshLayout(false);
 }
 
 void FAnimationMixerLayerModel::OnConstruct()
@@ -142,6 +176,12 @@ void FAnimationMixerLayerModel::OnConstruct()
 			if (WeakChildTrack.IsValid())
 			{
 				TrackEditor = Sequencer->GetTrackEditor(WeakChildTrack.Get());
+
+				// Listen for child track changes to sync decoration models
+				// (handles decoration add/remove and reconstruction needs like Mode changes).
+				UMovieSceneTrack* ChildTrack = WeakChildTrack.Get();
+				ChildTrackChangedHandle = ChildTrack->OnSignatureChanged().AddSP(
+					this, &FAnimationMixerLayerModel::OnChildTrackChanged);
 			}
 			else if (WeakParentTrack.IsValid())
 			{
@@ -287,6 +327,16 @@ void FAnimationMixerLayerModel::RefreshLayout(bool bChildrenNeedLayout)
 				}
 			}
 		}
+
+		// Sync child track decorations (e.g., root motion settings) to show in the outliner
+		if (FDecorationModelStorageExtension* DecorationModelStorage = SequenceModel->CastDynamic<FDecorationModelStorageExtension>())
+		{
+			UMovieSceneAnimationMixerTrack* ParentTrack = WeakParentTrack.Get();
+			if (ParentTrack)
+			{
+				DecorationModelStorage->SyncDecorationModels(ChildTrack, ParentTrack, OutlinerChildren);
+			}
+		}
 	}
 	// Layer has regular sections - act like FTrackRowModel
 	else
@@ -501,6 +551,23 @@ TSharedPtr<SWidget> FAnimationMixerLayerModel::CreateOutlinerViewForColumn(const
 {
 	if (TrackEditor)
 	{
+		// For child track layers, route the "Add" column through the mixer track editor
+		// so it can show both child track items and decoration entries (e.g., root motion settings)
+		if (InColumnName == FCommonOutlinerNames::Add && WeakChildTrack.IsValid() && WeakParentTrack.IsValid())
+		{
+			TSharedPtr<FSequenceModel> SequenceModel = FindAncestorOfType<FSequenceModel>();
+			TSharedPtr<ISequencer> Sequencer = SequenceModel ? SequenceModel->GetSequencer() : nullptr;
+			if (Sequencer)
+			{
+				TSharedPtr<ISequencerTrackEditor> MixerTrackEditor = Sequencer->GetTrackEditor(WeakParentTrack.Get());
+				if (MixerTrackEditor)
+				{
+					FBuildColumnWidgetParams Params(SharedThis(this), InParams);
+					return MixerTrackEditor->BuildOutlinerColumnWidget(Params, InColumnName);
+				}
+			}
+		}
+
 		FBuildColumnWidgetParams Params(SharedThis(this), InParams);
 		return TrackEditor->BuildOutlinerColumnWidget(Params, InColumnName);
 	}
diff --git a/Engine/Plugins/Experimental/MovieSceneAnimMixer/Source/MovieSceneAnimMixerEditor/Private/MovieSceneAnimationMixerTrackEditor.cpp b/Engine/Plugins/Experimental/MovieSceneAnimMixer/Source/MovieSceneAnimMixerEditor/Private/MovieSceneAnimationMixerTrackEditor.cpp
--- a/Engine/Plugins/Experimental/MovieSceneAnimMixer/Source/MovieSceneAnimMixerEditor/Private/MovieSceneAnimationMixerTrackEditor.cpp
+++ b/Engine/Plugins/Experimental/MovieSceneAnimMixer/Source/MovieSceneAnimMixerEditor/Private/MovieSceneAnimationMixerTrackEditor.cpp
@@ -659,8 +659,22 @@ TSharedPtr<SWidget> FAnimationMixerTrackEditor::BuildOutlinerColumnWidget(const
 				UMovieSceneAnimationMixerLayer* Layer = MixerTrack->GetLayer(LayerRowIndex);
 				if (Layer && Layer->HasChildTrack())
 				{
-					// Layer has a child track - don't show '+' button
-					return SNullWidget::NullWidget;
+					// For child track layers, show '+' button only if there are compatible decorations
+					UMovieSceneTrack* ChildTrack = Layer->GetChildTrack();
+					if (!ChildTrack || !SequencerHelpers::HasCompatibleDecorations(ChildTrack))
+					{
+						return SNullWidget::NullWidget;
+					}
+
+					return MakeAddButton(
+						LOCTEXT("AddModifier", "Modifier"),
+						FOnGetContent::CreateSP(
+							this,
+							&FAnimationMixerTrackEditor::BuildAddSectionSubMenu,
+							TWeakViewModelPtr<IOutlinerExtension>(OutlinerItem),
+							TWeakViewModelPtr<FSequencerEditorViewModel>(Editor)
+						),
+						Params.ViewModel);
 				}
 			}
 		}
@@ -1040,6 +1054,7 @@ TSharedRef<SWidget> FAnimationMixerTrackEditor::BuildAddSectionSubMenu(TWeakView
 	int32 TargetRowIndex = INDEX_NONE;
 	bool bCanAddChildTracks = true;
 
+	bool bHasChildTrack = false;
 	if (LayerExtension)
 	{
 		// Adding to a specific layer - use the layer's row index
@@ -1051,6 +1066,10 @@ TSharedRef<SWidget> FAnimationMixerTrackEditor::BuildAddSectionSubMenu(TWeakView
 		{
 			bCanAddChildTracks = false;
 		}
+		if (Layer && Layer->HasChildTrack())
+		{
+			bHasChildTrack = true;
+		}
 	}
 	else
 	{
@@ -1058,8 +1077,8 @@ TSharedRef<SWidget> FAnimationMixerTrackEditor::BuildAddSectionSubMenu(TWeakView
 		TargetRowIndex = INDEX_NONE;
 	}
 
-	// Animation items
-	if (AllItemClasses.Num() > 0 || Providers.Num() > 0)
+	// Animation items - skip for child track layers since the child track occupies the entire layer
+	if (!bHasChildTrack && (AllItemClasses.Num() > 0 || Providers.Num() > 0))
 	{
 		MenuBuilder.BeginSection(NAME_None, LOCTEXT("AnimationCategoryLabel", "Animation:"));
 
@@ -1153,13 +1172,23 @@ TSharedRef<SWidget> FAnimationMixerTrackEditor::BuildAddSectionSubMenu(TWeakView
 		UMovieSceneAnimationMixerLayer* Layer = MixerTrack->GetLayer(LayerExtension->GetRowIndex());
 		if (Layer)
 		{
-			SequencerHelpers::BuildDecorationMenu(MenuBuilder, Layer, ObjectBindingGuid, Sequencer);
+			SequencerHelpers::BuildDecorationMenu(MenuBuilder, Layer, ObjectBindingGuid, Sequencer,
+				LOCTEXT("LayerModifiers", "Layer Modifiers"));
 
-			// For single-section layers, also show decorations compatible with the section
-			TArray<UMovieSceneSection*> LayerSections = Layer->GetSections();
-			if (LayerSections.Num() == 1 && LayerSections[0])
+			if (Layer->HasChildTrack())
 			{
-				SequencerHelpers::BuildDecorationMenu(MenuBuilder, LayerSections[0], ObjectBindingGuid, Sequencer);
+				// Show decorations compatible with the child track (e.g., root motion settings)
+				SequencerHelpers::BuildDecorationMenu(MenuBuilder, Layer->GetChildTrack(), ObjectBindingGuid, Sequencer,
+					LOCTEXT("TrackModifiers", "Track Modifiers"));
+			}
+			else
+			{
+				// For single-section layers, also show decorations compatible with the section
+				TArray<UMovieSceneSection*> LayerSections = Layer->GetSections();
+				if (LayerSections.Num() == 1 && LayerSections[0])
+				{
+					SequencerHelpers::BuildDecorationMenu(MenuBuilder, LayerSections[0], ObjectBindingGuid, Sequencer);
+				}
 			}
 		}
 	}
diff --git a/Engine/Plugins/Experimental/MovieSceneAnimMixer/Source/MovieSceneAnimMixerEditor/Private/RootMotionSettingsDecorationEditor.cpp b/Engine/Plugins/Experimental/MovieSceneAnimMixer/Source/MovieSceneAnimMixerEditor/Private/RootMotionSettingsDecorationEditor.cpp
--- a/Engine/Plugins/Experimental/MovieSceneAnimMixer/Source/MovieSceneAnimMixerEditor/Private/RootMotionSettingsDecorationEditor.cpp
+++ b/Engine/Plugins/Experimental/MovieSceneAnimMixer/Source/MovieSceneAnimMixerEditor/Private/RootMotionSettingsDecorationEditor.cpp
@@ -293,9 +293,9 @@ void FRootMotionSettingsDecorationEditor::PopulateRootTransformMenu(
 			UMovieSceneSection* Section = WeakSection.Get();
 			if (Settings && Section)
 			{
+				Settings->InvalidateDecorationChannelProxy();
 				Settings->Modify();
 				Settings->TransformMode.SetDefault((uint8)Mode);
-				Settings->InvalidateDecorationChannelProxy();
 				if (TSharedPtr<ISequencer> SequencerPtr = Sequencer.Pin())
 				{
 					SequencerPtr->NotifyMovieSceneDataChanged(EMovieSceneDataChangeType::MovieSceneStructureItemsChanged);
@@ -756,9 +756,10 @@ TSharedPtr<SWidget> FRootMotionSettingsDecorationEditor::CreateKeyEditor(
 				}
 
 				FScopedTransaction Transaction(LOCTEXT("SetTransformMode", "Set Transform Mode"));
+				// Invalidate channel proxy before Modify() so listeners see the dirty flag.
+				Settings->InvalidateDecorationChannelProxy();
 				Settings->Modify();
 				Settings->TransformMode.SetDefault((uint8)Selection);
-				Settings->InvalidateDecorationChannelProxy();
 				if (TSharedPtr<ISequencer> SequencerPtr = WeakSequencer.Pin())
 				{
 					SequencerPtr->NotifyMovieSceneDataChanged(EMovieSceneDataChangeType::MovieSceneStructureItemsChanged);
diff --git a/Engine/Plugins/Experimental/MovieSceneAnimMixer/Source/MovieSceneAnimMixerEditor/Public/MVVM/AnimationMixerLayerModel.h b/Engine/Plugins/Experimental/MovieSceneAnimMixer/Source/MovieSceneAnimMixerEditor/Public/MVVM/AnimationMixerLayerModel.h
--- a/Engine/Plugins/Experimental/MovieSceneAnimMixer/Source/MovieSceneAnimMixerEditor/Public/MVVM/AnimationMixerLayerModel.h
+++ b/Engine/Plugins/Experimental/MovieSceneAnimMixer/Source/MovieSceneAnimMixerEditor/Public/MVVM/AnimationMixerLayerModel.h
@@ -215,6 +215,10 @@ private:
 	TSharedPtr<ISequencerTrackEditor> TrackEditor;
 
 	TOptional<int32> PreviousLayoutNumSections = TOptional<int32>();
+
+	/** Handles lightweight decoration sync when child track changes */
+	void OnChildTrackChanged();
+	FDelegateHandle ChildTrackChangedHandle;
 };
 
 } // namespace UE::Sequencer
diff --git a/Engine/Plugins/Experimental/MovieSceneMixedControlRig/Source/MovieSceneMixedControlRig/Private/MovieSceneMixedControlRigModule.cpp b/Engine/Plugins/Experimental/MovieSceneMixedControlRig/Source/MovieSceneMixedControlRig/Private/MovieSceneMixedControlRigModule.cpp
--- a/Engine/Plugins/Experimental/MovieSceneMixedControlRig/Source/MovieSceneMixedControlRig/Private/MovieSceneMixedControlRigModule.cpp
+++ b/Engine/Plugins/Experimental/MovieSceneMixedControlRig/Source/MovieSceneMixedControlRig/Private/MovieSceneMixedControlRigModule.cpp
@@ -4,6 +4,7 @@
 #include "Modules/ModuleInterface.h"
 #include "Modules/ModuleManager.h"
 #include "IMovieSceneModule.h"
+#include "MovieSceneAnimationMixerTrack.h"
 #include "MovieSceneMirroringDecoration.h"
 #include "Sequencer/MovieSceneControlRigParameterTrack.h"
 
@@ -19,6 +20,10 @@ void FMovieSceneMixedControlRigModule::StartupModule()
 	IMovieSceneModule::Get().RegisterCompatibleDecoration(
 		UMovieSceneControlRigParameterTrack::StaticClass(),
 		UMovieSceneMirroringDecoration::StaticClass());
+
+	IMovieSceneModule::Get().RegisterCompatibleDecoration(
+		UMovieSceneControlRigParameterTrack::StaticClass(),
+		UMovieSceneRootMotionSettingsDecoration::StaticClass());
 }
 
 void FMovieSceneMixedControlRigModule::ShutdownModule()
@@ -28,6 +33,10 @@ void FMovieSceneMixedControlRigModule::ShutdownModule()
 		IMovieSceneModule::Get().UnregisterCompatibleDecoration(
 			UMovieSceneControlRigParameterTrack::StaticClass(),
 			UMovieSceneMirroringDecoration::StaticClass());
+
+		IMovieSceneModule::Get().UnregisterCompatibleDecoration(
+			UMovieSceneControlRigParameterTrack::StaticClass(),
+			UMovieSceneRootMotionSettingsDecoration::StaticClass());
 	}
 }
 
diff --git a/Engine/Source/Editor/Sequencer/Private/IKeyArea.cpp b/Engine/Source/Editor/Sequencer/Private/IKeyArea.cpp
--- a/Engine/Source/Editor/Sequencer/Private/IKeyArea.cpp
+++ b/Engine/Source/Editor/Sequencer/Private/IKeyArea.cpp
@@ -98,6 +98,14 @@ FKeyHandle IKeyArea::AddOrUpdateKey(FFrameNumber Time, const FGuid& ObjectBindin
 	FMovieSceneChannel* Channel = ChannelHandle.Get();
 	UMovieSceneSection* Section = GetOwningSection();
 	check(Section != nullptr);
+
+	// Decoration-owned channels route through the decoration editor for keying.
+	auto [DecorationEditor, Decoration] = ISequencerDecorationEditor::FindDecorationOwner(ChannelHandle, Section);
+	if (DecorationEditor)
+	{
+		return DecorationEditor->AddOrUpdateKey(*Decoration, *Section, ChannelHandle, Time, ObjectBindingID, InSequencer);
+	}
+
 	UMovieSceneTrack* Track = CastChecked<UMovieSceneTrack>(Section->GetOuter());
 
 	// If track does not allow auto expanding sections, check if key time is within bounds. We include exclusive end frame to allow keyframes to visually align with section rendering
diff --git a/Engine/Source/Editor/Sequencer/Private/ISequencerDecorationEditor.cpp b/Engine/Source/Editor/Sequencer/Private/ISequencerDecorationEditor.cpp
new file mode 100644
--- /dev/null
+++ b/Engine/Source/Editor/Sequencer/Private/ISequencerDecorationEditor.cpp
@@ -0,0 +1,65 @@
+// Copyright Epic Games, Inc. All Rights Reserved.
+
+#include "ISequencerDecorationEditor.h"
+
+#include "ISequencerModule.h"
+#include "ISequencerChannelInterface.h"
+#include "Channels/MovieSceneChannel.h"
+#include "Channels/MovieSceneChannelProxy.h"
+#include "MovieSceneSection.h"
+#include "MovieSceneSignedObject.h"
+
+TTuple<ISequencerDecorationEditor*, UObject*> ISequencerDecorationEditor::FindDecorationOwner(
+	const FMovieSceneChannelHandle& ChannelHandle,
+	UMovieSceneSection* Section)
+{
+	if (const FMovieSceneChannelMetaData* MetaData = ChannelHandle.GetMetaData())
+	{
+		UObject* OwningObject = MetaData->WeakOwningObject.Get();
+		if (OwningObject && OwningObject != static_cast<UObject*>(Section))
+		{
+			if (ISequencerDecorationEditor* Editor = FindDecorationEditor(OwningObject->GetClass()))
+			{
+				return MakeTuple(Editor, OwningObject);
+			}
+		}
+	}
+	return MakeTuple(nullptr, nullptr);
+}
+
+FKeyHandle ISequencerDecorationEditor::AddOrUpdateKey(
+	UObject& Decoration,
+	UMovieSceneSection& Section,
+	FMovieSceneChannelHandle ChannelHandle,
+	FFrameNumber Time,
+	const FGuid& ObjectBindingID,
+	ISequencer& InSequencer)
+{
+	FMovieSceneChannel* Channel = ChannelHandle.Get();
+	if (!Channel)
+	{
+		return FKeyHandle();
+	}
+
+	// Apply time offset from channel metadata
+	if (const FMovieSceneChannelMetaData* MetaData = ChannelHandle.GetMetaData())
+	{
+		Time -= MetaData->GetOffsetTime(&Section);
+	}
+
+	// Find the channel editor interface for this channel type
+	ISequencerModule& SequencerModule = FModuleManager::LoadModuleChecked<ISequencerModule>("Sequencer");
+	ISequencerChannelInterface* EditorInterface = SequencerModule.FindChannelEditorInterface(ChannelHandle.GetChannelTypeName());
+
+	if (EditorInterface)
+	{
+		if (UMovieSceneSignedObject* SignedDecoration = Cast<UMovieSceneSignedObject>(&Decoration))
+		{
+			SignedDecoration->Modify();
+		}
+		const void* RawExtendedData = ChannelHandle.GetExtendedEditorData();
+		return EditorInterface->AddOrUpdateKey_Raw(Channel, &Section, RawExtendedData, Time, InSequencer, ObjectBindingID, nullptr);
+	}
+
+	return FKeyHandle();
+}
diff --git a/Engine/Source/Editor/Sequencer/Private/MVVM/ViewModels/ChannelModel.cpp b/Engine/Source/Editor/Sequencer/Private/MVVM/ViewModels/ChannelModel.cpp
--- a/Engine/Source/Editor/Sequencer/Private/MVVM/ViewModels/ChannelModel.cpp
+++ b/Engine/Source/Editor/Sequencer/Private/MVVM/ViewModels/ChannelModel.cpp
@@ -262,9 +262,33 @@ bool FChannelModel::UpdateCachedKeys(TSharedPtr<FCachedKeys>& OutCachedKeys) con
 		bool Update(const FChannelModel* InChannel)
 		{
 			UMovieSceneSection* Section = InChannel->GetSection();
-			if (!Section || !CachedSignature.IsValid() || Section->GetSignature() != CachedSignature)
+
+			// Check both the section signature and the owning object signature.
+			// For decoration channels, the owning object (decoration) changes when
+			// keys are modified, but the host section's signature may not update
+			// since it's not in the decoration's outer chain.
+			FGuid OwnerSignature;
+			if (const TSharedPtr<IKeyArea>& ChannelKeyArea = InChannel->GetKeyArea())
+			{
+				if (const FMovieSceneChannelMetaData* MetaData = ChannelKeyArea->GetChannel().GetMetaData())
+				{
+					if (UMovieSceneSignedObject* SignedOwner = Cast<UMovieSceneSignedObject>(MetaData->WeakOwningObject.Get()))
+					{
+						if (SignedOwner != Section)
+						{
+							OwnerSignature = SignedOwner->GetSignature();
+						}
+					}
+				}
+			}
+
+			const bool bSectionChanged = !Section || !CachedSignature.IsValid() || Section->GetSignature() != CachedSignature;
+			const bool bOwnerChanged = OwnerSignature.IsValid() && OwnerSignature != CachedOwnerSignature;
+
+			if (bSectionChanged || bOwnerChanged)
 			{
 				CachedSignature = Section && InChannel->GetKeyArea() && InChannel->GetKeyArea()->ResolveChannel() ? Section->GetSignature() : FGuid();
+				CachedOwnerSignature = OwnerSignature;
 
 				KeyTimes.Reset();
 				KeyHandles.Reset();
@@ -283,8 +307,8 @@ bool FChannelModel::UpdateCachedKeys(TSharedPtr<FCachedKeys>& OutCachedKeys) con
 			return false;
 		}
 
-		/** The guid with which the above array was generated */
 		FGuid CachedSignature;
+		FGuid CachedOwnerSignature;
 	};
 
 	if (OutCachedKeys)
@@ -1070,6 +1094,8 @@ TSharedPtr<SWidget> FChannelGroupOutlinerModel::CreateOutlinerViewForColumn(cons
 								return CustomWidget;
 							}
 						}
+
+						// Fall through to standard key navigation.
 					}
 				}
 			}
diff --git a/Engine/Source/Editor/Sequencer/Private/MVVM/ViewModels/DecorationOutlinerModel.cpp b/Engine/Source/Editor/Sequencer/Private/MVVM/ViewModels/DecorationOutlinerModel.cpp
--- a/Engine/Source/Editor/Sequencer/Private/MVVM/ViewModels/DecorationOutlinerModel.cpp
+++ b/Engine/Source/Editor/Sequencer/Private/MVVM/ViewModels/DecorationOutlinerModel.cpp
@@ -4,6 +4,7 @@
 
 #include "ISequencer.h"
 #include "ISequencerDecorationEditor.h"
+#include "ISequencerSection.h"
 #include "SequencerLog.h"
 #include "MovieSceneSection.h"
 #include "SequencerCommonHelpers.h"
@@ -135,7 +136,10 @@ void FDecorationOutlinerModel::OnConstruct()
 
 	// Handle section-provider decorations: host sections directly in this model's track
 	// area without creating intermediate FSectionOutlinerModel wrappers.
-	if (IMovieSceneSectionProviderDecoration* SectionProvider = Cast<IMovieSceneSectionProviderDecoration>(Decoration))
+	// Skip this path if the decoration also implements IMovieSceneChannelDecoration,
+	// since the channel path handles both section and track containers.
+	if (IMovieSceneSectionProviderDecoration* SectionProvider = Cast<IMovieSceneSectionProviderDecoration>(Decoration);
+		SectionProvider && !Cast<IMovieSceneChannelDecoration>(Decoration))
 	{
 		ISequencerDecorationEditor* DecorationEditor = ISequencerDecorationEditor::FindDecorationEditor(Decoration->GetClass());
 		if (!DecorationEditor)
@@ -177,23 +181,50 @@ void FDecorationOutlinerModel::OnConstruct()
 		}
 
 		UMovieSceneSection* ParentSection = Cast<UMovieSceneSection>(WeakDecorationContainer.Get());
-		if (!ParentSection)
+		TSharedPtr<FSectionModel> SectionModel;
+		FMovieSceneChannelProxy* ChannelProxyPtr = nullptr;
+
+		if (ParentSection)
 		{
-			return;
+			// Section-level decoration: channels are hosted by the parent section,
+			// with a separate proxy owned by this model.
+			TSharedRef<ISequencerSection> NewSectionInterface = TrackEditor->MakeSectionInterface(
+				*ParentSection, *WeakParentTrack.Get(), WeakParentTrack.Get()->FindObjectBindingGuid());
+			SectionModel = SectionModelStorage->CreateModelForSection(ParentSection, NewSectionInterface);
+
+			FMovieSceneChannelProxyData ChannelProxyData;
+			ChannelDecoration->PopulateChannelProxy(ChannelProxyData);
+			ChannelDecorationProxy = MakeShared<FMovieSceneChannelProxy>(MoveTemp(ChannelProxyData));
+			ChannelProxyPtr = ChannelDecorationProxy.Get();
 		}
+		else
+		{
+			// Track-level decoration: get the host section from the decoration's
+			// IMovieSceneSectionProviderDecoration and use it for channel hosting.
+			if (IMovieSceneSectionProviderDecoration* HostSectionProvider = Cast<IMovieSceneSectionProviderDecoration>(Decoration))
+			{
+				TArrayView<TObjectPtr<UMovieSceneSection>> Sections = HostSectionProvider->GetSections();
+				if (Sections.Num() > 0)
+				{
+					ParentSection = Sections[0];
+				}
+			}
+			if (!ParentSection)
+			{
+				return;
+			}
 
-		TSharedRef<ISequencerSection> NewSectionInterface = TrackEditor->MakeSectionInterface(*ParentSection, *WeakParentTrack.Get(), WeakParentTrack.Get()->FindObjectBindingGuid());
-		const TSharedPtr<FSectionModel> SectionModel = SectionModelStorage->CreateModelForSection(ParentSection, NewSectionInterface);
-
-		TSharedPtr<ISequencerSection> SectionInterface = SectionModel->GetSectionInterface();
+			TSharedRef<ISequencerSection> NewSectionInterface = MakeShared<FSequencerSection>(*ParentSection);
+			SectionModel = SectionModelStorage->CreateModelForSection(ParentSection, NewSectionInterface);
+			SectionChildren.AddChild(SectionModel);
+			SectionModel->SetLinkedOutlinerItem(CastViewModel<IOutlinerExtension>(AsShared()));
 
-		// Get channels from the decoration
-		FMovieSceneChannelProxyData ChannelProxyData;
-		ChannelDecoration->PopulateChannelProxy(ChannelProxyData);
+			ParentSection->InvalidateChannelProxy();
+			ChannelProxyPtr = &ParentSection->GetChannelProxy();
+		}
 
-		// Build channel proxy and store it to keep channel handles valid
-		ChannelDecorationProxy = MakeShared<FMovieSceneChannelProxy>(MoveTemp(ChannelProxyData));
-		FMovieSceneChannelProxy& ChannelProxy = *ChannelDecorationProxy;
+		FMovieSceneChannelProxy& ChannelProxy = *ChannelProxyPtr;
+		TSharedPtr<ISequencerSection> SectionInterface = SectionModel->GetSectionInterface();
 
 		// Collect all channels with their metadata for sorting
 		struct FChannelSortData
@@ -454,7 +485,10 @@ FViewModelChildren FDecorationOutlinerModel::GetTopLevelChannels()
 FTrackAreaParameters FDecorationOutlinerModel::GetTrackAreaParameters() const
 {
 	FTrackAreaParameters Parameters;
-	Parameters.LaneType = ETrackAreaLaneType::Inline;
+	// Use Nested when the decoration has its own section so child channel lanes
+	// register as children, enabling section background painting behind them.
+	const bool bHasOwnSection = SectionModelList.Iterate<FSectionModel>().begin() != SectionModelList.Iterate<FSectionModel>().end();
+	Parameters.LaneType = bHasOwnSection ? ETrackAreaLaneType::Nested : ETrackAreaLaneType::Inline;
 	return Parameters;
 }
 
diff --git a/Engine/Source/Editor/Sequencer/Private/SequencerAddKeyOperation.cpp b/Engine/Source/Editor/Sequencer/Private/SequencerAddKeyOperation.cpp
--- a/Engine/Source/Editor/Sequencer/Private/SequencerAddKeyOperation.cpp
+++ b/Engine/Source/Editor/Sequencer/Private/SequencerAddKeyOperation.cpp
@@ -20,6 +20,7 @@
 #include "MovieSceneSection.h"
 #include "MovieSceneTrack.h"
 #include "SequencerKeyParams.h"
+#include "ISequencerDecorationEditor.h"
 #include "Templates/Tuple.h"
 #include "Templates/TypeHash.h"
 #include "Templates/UnrealTemplate.h"
@@ -170,6 +171,14 @@ bool FAddKeyOperation::ProcessKeyArea(ISequencerTrackEditor* InTrackEditor, TVie
 
 bool FAddKeyOperation::ProcessKeyArea(ISequencerTrackEditor* InTrackEditor, TSharedPtr<IKeyArea> InKeyArea)
 {
+	// Decoration-owned channels are keyed through the decoration editor.
+	auto [DecorationEditor, Decoration] = ISequencerDecorationEditor::FindDecorationOwner(InKeyArea->GetChannel(), InKeyArea->GetOwningSection());
+	if (DecorationEditor)
+	{
+		DecorationKeyOperations.Add({DecorationEditor, Decoration, InKeyArea->GetOwningSection(), InKeyArea->GetChannel()});
+		return true;
+	}
+
 	TSharedPtr<ISequencerSection> Section       = InKeyArea->GetSectionInterface();
 	UMovieSceneSection*           SectionObject = Section       ? Section->GetSectionObject()                      : nullptr;
 	UMovieSceneTrack*             TrackObject   = SectionObject ? SectionObject->GetTypedOuter<UMovieSceneTrack>() : nullptr;
@@ -191,6 +200,12 @@ void FAddKeyOperation::Commit(FFrameNumber KeyTime, ISequencer& InSequencer, TAr
 		Pair.Key->ProcessKeyOperation(KeyTime, Pair.Value, InSequencer, OutResults);
 	}
 
+	// Process decoration key operations collected during ProcessKeyArea
+	for (const FDecorationKeyOp& Op : DecorationKeyOperations)
+	{
+		Op.Editor->AddOrUpdateKey(*Op.Decoration, *Op.Section, Op.ChannelHandle, KeyTime, FGuid(), InSequencer);
+	}
+
 	InSequencer.UpdatePlaybackRange();
 	InSequencer.NotifyMovieSceneDataChanged(EMovieSceneDataChangeType::TrackValueChanged);
 }
diff --git a/Engine/Source/Editor/Sequencer/Private/SequencerCommonHelpers.cpp b/Engine/Source/Editor/Sequencer/Private/SequencerCommonHelpers.cpp
--- a/Engine/Source/Editor/Sequencer/Private/SequencerCommonHelpers.cpp
+++ b/Engine/Source/Editor/Sequencer/Private/SequencerCommonHelpers.cpp
@@ -1196,7 +1196,8 @@ void SequencerHelpers::BuildDecorationMenu(
 	FMenuBuilder& MenuBuilder,
 	UMovieSceneDecorationContainerObject* DecorationContainer,
 	const FGuid& ObjectBindingID,
-	TWeakPtr<ISequencer> InSequencer)
+	TWeakPtr<ISequencer> InSequencer,
+	const FText& SectionLabel)
 {
 	if (!DecorationContainer)
 	{
@@ -1231,7 +1232,8 @@ void SequencerHelpers::BuildDecorationMenu(
 		return A.GetDecorationMenuPriority() < B.GetDecorationMenuPriority();
 	});
 
-	MenuBuilder.BeginSection(NAME_None, NSLOCTEXT("SequencerHelpers", "Modifiers_Label", "Modifiers"));
+	const FText& Label = SectionLabel.IsEmpty() ? NSLOCTEXT("SequencerHelpers", "Modifiers_Label", "Modifiers") : SectionLabel;
+	MenuBuilder.BeginSection(NAME_None, Label);
 	{
 
 		// First, allow menu providers to handle decorations.
diff --git a/Engine/Source/Editor/Sequencer/Public/ISequencerDecorationEditor.h b/Engine/Source/Editor/Sequencer/Public/ISequencerDecorationEditor.h
--- a/Engine/Source/Editor/Sequencer/Public/ISequencerDecorationEditor.h
+++ b/Engine/Source/Editor/Sequencer/Public/ISequencerDecorationEditor.h
@@ -4,6 +4,7 @@
 
 #include "CoreMinimal.h"
 #include "Channels/MovieSceneChannelHandle.h"
+#include "Curves/KeyHandle.h"
 #include "Features/IModularFeatures.h"
 
 class FMenuBuilder;
@@ -98,6 +99,28 @@ public:
 		const FGuid& ObjectBindingID,
 		TSharedPtr<UE::Sequencer::FViewModel> ChannelGroupModel) { return nullptr; }
 
+	/**
+	 * Add or update a key on a decoration-owned channel. Called when the standard
+	 * keying path detects a channel owned by a decoration rather than a track.
+	 * The default implementation keys the channel directly without track assumptions.
+	 */
+	SEQUENCER_API virtual FKeyHandle AddOrUpdateKey(
+		UObject& Decoration,
+		UMovieSceneSection& Section,
+		FMovieSceneChannelHandle ChannelHandle,
+		FFrameNumber Time,
+		const FGuid& ObjectBindingID,
+		ISequencer& InSequencer);
+
+	/**
+	 * Find the decoration editor and owning decoration for a channel that may be
+	 * owned by a decoration rather than its section. Returns {nullptr, nullptr}
+	 * if the channel is section-owned or has no registered decoration editor.
+	 */
+	SEQUENCER_API static TTuple<ISequencerDecorationEditor*, UObject*> FindDecorationOwner(
+		const FMovieSceneChannelHandle& ChannelHandle,
+		UMovieSceneSection* Section);
+
 	/**
 	 * Get the modular feature name for decoration editors.
 	 */
diff --git a/Engine/Source/Editor/Sequencer/Public/SequencerAddKeyOperation.h b/Engine/Source/Editor/Sequencer/Public/SequencerAddKeyOperation.h
--- a/Engine/Source/Editor/Sequencer/Public/SequencerAddKeyOperation.h
+++ b/Engine/Source/Editor/Sequencer/Public/SequencerAddKeyOperation.h
@@ -7,6 +7,7 @@
 #include "Containers/Map.h"
 #include "Containers/Set.h"
 #include "Containers/SparseArray.h"
+#include "Channels/MovieSceneChannelHandle.h"
 #include "SequencerCoreFwd.h"
 #include "Misc/FrameNumber.h"
 #include "SequencerKeyParams.h"
@@ -17,7 +18,9 @@
 class IKeyArea;
 class ISequencer;
 class ISequencerSection;
+class ISequencerDecorationEditor;
 class ISequencerTrackEditor;
+class UMovieSceneSection;
 class UMovieSceneTrack;
 namespace UE::Sequencer { struct FKeyOperation; }
 struct FFrameNumber;
@@ -116,6 +119,16 @@ private:
 
 	/** Map of key operations stored by their track editor. */
 	TMap<ISequencerTrackEditor*, FKeyOperation> OperationsByTrackEditor;
+
+	/** Pending decoration key operations collected during ProcessKeyArea. */
+	struct FDecorationKeyOp
+	{
+		ISequencerDecorationEditor* Editor;
+		UObject* Decoration;
+		UMovieSceneSection* Section;
+		FMovieSceneChannelHandle ChannelHandle;
+	};
+	TArray<FDecorationKeyOp> DecorationKeyOperations;
 };
 
 } // namespace Sequencer
diff --git a/Engine/Source/Editor/Sequencer/Public/SequencerCommonHelpers.h b/Engine/Source/Editor/Sequencer/Public/SequencerCommonHelpers.h
--- a/Engine/Source/Editor/Sequencer/Public/SequencerCommonHelpers.h
+++ b/Engine/Source/Editor/Sequencer/Public/SequencerCommonHelpers.h
@@ -161,7 +161,8 @@ public:
 	static void BuildDecorationMenu(FMenuBuilder& MenuBuilder
 		, UMovieSceneDecorationContainerObject* DecorationContainer
 		, const FGuid& ObjectBindingID
-		, TWeakPtr<ISequencer> InSequencer);
+		, TWeakPtr<ISequencer> InSequencer
+		, const FText& SectionLabel = FText());
 
 	/*
 	 * Build an FPropertyPath from an IPropertyHandle. 
