[Backout] - CL52249771
#rnx
Original CL Desc
-----------------------------------------------------------------
Sequencer: Adjust filtering algorithm to improve performance when filtering using the Selected filter. Future improvements can be made to other filters later.

The main idea is to avoid search recoursing the FViewModel tree hiearchy needlessly. Filters can now specify whether their result applies to all children, too. For example, for the selection filter, we need not iterate all children of a selected node: they should all be shown so it is pointless to rerun filters on those nodes.

Main changes:
- EItemFilterState: Include, Exclude, DoNotCare.
- FSequencerTrackFilter::Evaluate: This function allows filters to whether the filter in (include), out (exclude), or do not affect an item AND whether the result should be applied to all children.
- FSequencerTrackFilter_Selected::OnPreFilter now pre-iterates the entire hierarchy once to determine the view models to include. This is massive performance boost: before the algorithm would iterate each tree item, and FSequencerTrackFilter_Selected would reiterate up the entire parent hierarchy to check whether one of its parents is selected.

The old algorithm continues to exist and can be enabled using "Sequencer.UseNewFilteringAlgorithm 0" console variable. The new algorithm is enabled by default. The point of this CVar is so in case anything is broken in the new filtering algorith, animators can revert back to the old algorithm locally while we address any potential bugs.

On a simple case with only the selected filter, 2 SM actors, 1 Mannequin, and 2 dragons, filtering went from 470ms to only 1.77ms.

diff --git a/Engine/Source/Editor/Sequencer/Private/Filters/Evaluation/FilterEvaluator.cpp b/Engine/Source/Editor/Sequencer/Private/Filters/Evaluation/FilterEvaluator.cpp
--- a/Engine/Source/Editor/Sequencer/Private/Filters/Evaluation/FilterEvaluator.cpp
+++ b/Engine/Source/Editor/Sequencer/Private/Filters/Evaluation/FilterEvaluator.cpp
@@ -40,14 +40,11 @@ bool FFilterEvaluator::UpdateFilters(bool bForce)
 	{
 		return false;
 	}
-	
-	TRACE_CPUPROFILER_EVENT_SCOPE(FFilterEvaluator::UpdateFilters);
 
 	bool bHasFilterDataChanged = false; 
 	for (FFilterBarData& FilterBarData : RegisteredFilterBars)
 	{
 		const TSharedRef<FSequencerFilterBar>& FilterBar = FilterBarData.FilterBar;
-		TRACE_CPUPROFILER_EVENT_SCOPE_TEXT(FilterBar->GetIdentifier());
 		
 		const FSequencerFilterData PreviousFilterData = FilterBar->GetFilterData();
 		const FSequencerFilterData& FilterData = FilterBar->FilterNodes();
diff --git a/Engine/Source/Editor/Sequencer/Private/Filters/Filters/SequencerTrackFilterBase.cpp b/Engine/Source/Editor/Sequencer/Private/Filters/Filters/SequencerTrackFilterBase.cpp
--- a/Engine/Source/Editor/Sequencer/Private/Filters/Filters/SequencerTrackFilterBase.cpp
+++ b/Engine/Source/Editor/Sequencer/Private/Filters/Filters/SequencerTrackFilterBase.cpp
@@ -23,13 +23,6 @@ ISequencerTrackFilters& FSequencerTrackFilter::GetFilterInterface() const
 	return static_cast<ISequencerTrackFilters&>(FilterInterface);
 }
 
-FFilterResult FSequencerTrackFilter::Evaluate(FSequencerTrackFilterType InItem) const
-{
-	// This API was added in 5.8. We allow filters to only implement PassesFilter.
-	const bool bPassesFilter = PassesFilter(InItem);
-	return FFilterResult(bPassesFilter ? EItemFilterState::Include : EItemFilterState::Exclude);
-}
-
 UMovieSceneSequence* FSequencerTrackFilter::GetFocusedMovieSceneSequence() const
 {
 	return GetSequencer().GetFocusedMovieSceneSequence();
diff --git a/Engine/Source/Editor/Sequencer/Private/Filters/Filters/SequencerTrackFilter_Selected.cpp b/Engine/Source/Editor/Sequencer/Private/Filters/Filters/SequencerTrackFilter_Selected.cpp
--- a/Engine/Source/Editor/Sequencer/Private/Filters/Filters/SequencerTrackFilter_Selected.cpp
+++ b/Engine/Source/Editor/Sequencer/Private/Filters/Filters/SequencerTrackFilter_Selected.cpp
@@ -5,7 +5,6 @@
 #include "Filters/SequencerFilterBar.h"
 #include "Filters/SequencerTrackFilterCommands.h"
 #include "Filters/SequencerTrackFilterSelectedExtension.h"
-#include "HAL/IConsoleManager.h"
 #include "ISequencer.h"
 #include "MVVM/Selection/Selection.h"
 #include "MVVM/ViewModels/SequencerEditorViewModel.h"
@@ -16,6 +15,8 @@ using namespace UE::Sequencer;
 
 #define LOCTEXT_NAMESPACE "SequencerTrackFilter_Selected"
 
+
+
 namespace UE::Sequencer::SelectionDetail
 {
 static TArray<TWeakObjectPtr<USequencerTrackFilterSelectedExtension>> GetSelectionExtensions()
@@ -41,6 +42,7 @@ static bool IsPartOfHierarchy(const FViewModelPtr& InItem, const FViewModelPtr&
 	{
 		return false;
 	}
+
 	for (FParentFirstChildIterator It(InSelectedNode, /*bIncludeThis=*/true); It; ++It)
 	{
 		if (*It == InItem)
@@ -51,6 +53,7 @@ static bool IsPartOfHierarchy(const FViewModelPtr& InItem, const FViewModelPtr&
 	
 	return false;
 }
+
 /** @return Whether InItem is a child of the FViewModels enumerated by USequencerTrackFilterSelectedExtension. */
 static bool IsPartOfExtendedSelection(
 	const FViewModelPtr& InItem, 
@@ -60,20 +63,20 @@ static bool IsPartOfExtendedSelection(
 {
 	using EBreakBehavior = USequencerTrackFilterSelectedExtension::EBreakBehavior;
 	const auto ProcessExtension = [&InItem, &InFilter](const TWeakObjectPtr<USequencerTrackFilterSelectedExtension>& WeakExtension)
+	{
+		if (!WeakExtension.IsValid())
 		{
-			if (!WeakExtension.IsValid())
-			{
-				return false;
-			}
+			return false;
+		}
 		
-			bool bIsPartOfSelection = false;
-			WeakExtension->EnumerateViewModelsConsideredAsSelected(InFilter, [&InItem, &bIsPartOfSelection](const FViewModelPtr& ConsiderSelected)
-			{
-				bIsPartOfSelection = IsPartOfHierarchy(InItem, ConsiderSelected);
-				return bIsPartOfSelection ? EBreakBehavior::Break : EBreakBehavior::Continue;
-			});
-			return bIsPartOfSelection;
-		};
+		bool bIsPartOfSelection = false;
+		WeakExtension->EnumerateViewModelsConsideredAsSelected(InFilter, [&InItem, &bIsPartOfSelection](const FViewModelPtr& ConsiderSelected)
+		{
+			bIsPartOfSelection = IsPartOfHierarchy(InItem, ConsiderSelected);
+			return bIsPartOfSelection ? EBreakBehavior::Break : EBreakBehavior::Continue;
+		});
+		return bIsPartOfSelection;
+	};
 	
 	return Algo::AnyOf(InWeakExtensions, ProcessExtension);
 }
@@ -86,14 +89,14 @@ FSequencerTrackFilter_Selected::FSequencerTrackFilter_Selected(ISequencerTrackFi
 
 FSequencerTrackFilter_Selected::~FSequencerTrackFilter_Selected()
 {
-	UnbindEvents();
+	UnbindSelectionChanged();
 }
 
 void FSequencerTrackFilter_Selected::BindSelectionChanged()
 {
 	if (!OnSelectionChangedHandle.IsValid())
 	{
-		OnSelectionChangedHandle = USelection::SelectionChangedEvent.AddRaw(this, &FSequencerTrackFilter_Selected::OnViewportSelectionChanged);
+		OnSelectionChangedHandle = USelection::SelectionChangedEvent.AddRaw(this, &FSequencerTrackFilter_Selected::OnSelectionChanged);
 	}
 }
 
@@ -127,24 +130,14 @@ void FSequencerTrackFilter_Selected::ActiveStateChanged(const bool bInActive)
 
 	if (bInActive)
 	{
-		BindEvents();
+		BindSelectionChanged();
 	}
 	else
 	{
-		UnbindEvents();
+		UnbindSelectionChanged();
 	}
 }
 
-FFilterResult FSequencerTrackFilter_Selected::Evaluate(FSequencerTrackFilterType InItem) const
-{
-	const EFilterResultFlags* Flags = IncludedViewModels.Find(InItem);
-	const bool bPassesFilter = Flags != nullptr;
-	
-	const EItemFilterState ResultState = bPassesFilter ? EItemFilterState::Include : EItemFilterState::Exclude;
-	const EFilterResultFlags ResultFlags = Flags ? *Flags : EFilterResultFlags::None;
-	return FFilterResult(ResultState, ResultFlags);
-}
-
 FText FSequencerTrackFilter_Selected::GetDisplayName() const
 {
 	return LOCTEXT("SequencerTrackFilter_Selected", "Selected");
@@ -155,69 +148,6 @@ FSlateIcon FSequencerTrackFilter_Selected::GetIcon() const
 	return FSlateIcon(FAppStyle::GetAppStyleSetName(), TEXT("Icons.SelectInViewport"));
 }
 
-void FSequencerTrackFilter_Selected::OnPreFilter()
-{
-	if (!bNeedsPrefilterRebuild)
-	{
-		return;
-	}
-	
-	TRACE_CPUPROFILER_EVENT_SCOPE(FSequencerTrackFilter_Selected::OnPreFilter);
-	bNeedsPrefilterRebuild = false;
-	IncludedViewModels.Empty(IncludedViewModels.Num());
-	ON_SCOPE_EXIT{ IncludedViewModels.Shrink(); };
-	
-	const TSharedPtr<FSequencerEditorViewModel> RootViewModel = GetSequencer().GetViewModel();
-	const TSharedPtr<FSequencerSelection> SequencerSelection = RootViewModel ? RootViewModel->GetSelection() : nullptr;
-	if (!SequencerSelection.IsValid())
-	{
-		return;
-	}
-	
-	const auto ProcessSelected = [this](const FViewModelPtr& ViewModel)
-	{
-		if (IncludedViewModels.Contains(ViewModel))
-		{
-			return;
-		}
-			
-		// If "Foo.A" is selected, make "Foo.A.Child", etc. filtered in...
-		IncludedViewModels.Add(ViewModel, EFilterResultFlags::ApplyResultToChildren);
-	
-		for (const FViewModelPtr& Child : FParentFirstChildIterator(ViewModel))
-		{
-			// ... an if the filtered is asked whether "Foo.A.Child" is filtered in, it should also say, yes.
-			IncludedViewModels.Add(Child, EFilterResultFlags::ApplyResultToChildren);
-		}
-			
-		for (const FViewModelPtr& Parent : FParentModelIterator(ViewModel))
-		{
-			// Parent should be visible but do not tell them to propagate to all their children. 
-			// If Foo.A is selected, then Foo should be filtered in, but not Foo.B.
-			IncludedViewModels.FindOrAdd(Parent);
-		}
-	};
-
-	for (const FViewModelPtr& ViewModel : SequencerSelection->Outliner)
-	{
-		ProcessSelected(ViewModel);
-	}
-	
-	for (const TWeakObjectPtr<USequencerTrackFilterSelectedExtension>& Extension : CachedWeakExtensions)
-	{
-		if (!Extension.IsValid())
-		{
-			continue;
-		}
-		
-		Extension->EnumerateViewModelsConsideredAsSelected(*this, [&ProcessSelected](const FViewModelPtr& ViewModel)
-		{
-			ProcessSelected(ViewModel);
-			return USequencerTrackFilterSelectedExtension::EBreakBehavior::Continue;
-		});
-	}
-}
-
 FString FSequencerTrackFilter_Selected::GetName() const
 {
 	return StaticName();
@@ -225,70 +155,35 @@ FString FSequencerTrackFilter_Selected::GetName() const
 
 bool FSequencerTrackFilter_Selected::PassesFilter(FSequencerTrackFilterType InItem) const
 {
-	// TODO UE-371829: Remove this CVar
-	const IConsoleVariable* Variable = IConsoleManager::Get().FindConsoleVariable(TEXT("Sequencer.UseNewFilteringAlgorithm"));
-	const bool bUseOldAlgorithm = !ensure(Variable) || !Variable->GetBool();
-	if (bUseOldAlgorithm)
+	const TSharedPtr<FSequencerEditorViewModel> ViewModel = GetSequencer().GetViewModel();
+	if (!ViewModel.IsValid())
 	{
-		const TSharedPtr<FSequencerEditorViewModel> ViewModel = GetSequencer().GetViewModel();
-		if (!ViewModel.IsValid())
-		{
-			return false;
-		}
- 
-		const TSharedPtr<FSequencerSelection> SequencerSelection = ViewModel->GetSelection();
-		if (!SequencerSelection.IsValid())
-		{
-			return false;
-		}
-		const bool bIsPartOfSequencerSelection = Algo::AnyOf(SequencerSelection->Outliner, [InItem](const FViewModelPtr& Selected)
-		{
-			return SelectionDetail::IsPartOfHierarchy(InItem, Selected);
-		});
-		return bIsPartOfSequencerSelection 
-			// Allow external logic to make us pretend like additional items are selected without actually modifying FSequencerSelection::Outliner.
-			|| SelectionDetail::IsPartOfExtendedSelection(InItem, *this, CachedWeakExtensions);
+		return false;
 	}
-	
-	return IncludedViewModels.Contains(InItem);
-}
 
-void FSequencerTrackFilter_Selected::BindEvents()
-{
-	BindSelectionChanged();
-	
-	const TSharedPtr<FSequencerEditorViewModel> RootViewModel = GetSequencer().GetViewModel();
-	const TSharedPtr<FSequencerSelection> SequencerSelection = RootViewModel ? RootViewModel->GetSelection() : nullptr;
-	if (SequencerSelection)
+	const TSharedPtr<FSequencerSelection> SequencerSelection = ViewModel->GetSelection();
+	if (!SequencerSelection.IsValid())
 	{
-		SequencerSelection->OnChanged.AddRaw(this, &FSequencerTrackFilter_Selected::OnSequencerSelectionChanged);
+		return false;
 	}
-}
 
-void FSequencerTrackFilter_Selected::UnbindEvents()
-{
-	UnbindSelectionChanged();
-	
-	const TSharedPtr<FSequencerEditorViewModel> RootViewModel = GetSequencer().GetViewModel();
-	const TSharedPtr<FSequencerSelection> SequencerSelection = RootViewModel ? RootViewModel->GetSelection() : nullptr;
-	if (SequencerSelection)
+	const bool bIsPartOfSequencerSelection = Algo::AnyOf(SequencerSelection->Outliner, [InItem](const FViewModelPtr& Selected)
 	{
-		SequencerSelection->OnChanged.RemoveAll(this);
-	}
+		return SelectionDetail::IsPartOfHierarchy(InItem, Selected);
+	});
+	return bIsPartOfSequencerSelection 
+		// Allow external logic to make us pretend like additional items are selected without actually modifying FSequencerSelection::Outliner.
+		|| SelectionDetail::IsPartOfExtendedSelection(InItem, *this, CachedWeakExtensions);
 }
 
-void FSequencerTrackFilter_Selected::OnSequencerSelectionChanged()
+void FSequencerTrackFilter_Selected::OnSelectionChanged(UObject* const InObject)
 {
-	bNeedsPrefilterRebuild = true;
-	// TODO UE-371115: For now, we achieve this ticket behavior partially by not refiltering. However, once the ticket is completed, this should be uncommented.
-	// FilterInterface.RequestFilterUpdate();
-	// Achieved partially because any subsequent filter reevaluation, like updating the search bar and resetting it, will hide previously multi-selected items.
+	FilterInterface.RequestFilterUpdate();
 }
 
-void FSequencerTrackFilter_Selected::OnViewportSelectionChanged(UObject* const InObject)
+void FSequencerTrackFilter_Selected::ToggleShowOnlySelectedTracks()
 {
-	bNeedsPrefilterRebuild = true;
-	FilterInterface.RequestFilterUpdate();
+	SetActive(!IsActive());
 }
 
 #undef LOCTEXT_NAMESPACE
diff --git a/Engine/Source/Editor/Sequencer/Private/Filters/SequencerFilterBar.cpp b/Engine/Source/Editor/Sequencer/Private/Filters/SequencerFilterBar.cpp
--- a/Engine/Source/Editor/Sequencer/Private/Filters/SequencerFilterBar.cpp
+++ b/Engine/Source/Editor/Sequencer/Private/Filters/SequencerFilterBar.cpp
@@ -33,7 +33,6 @@
 #include "Filters/SequencerFilterAreaConfig.h"
 #include "Stats/StatsMisc.h"
 #include "UObject/UObjectIterator.h"
-#include "Utils/FilterEvaluation.h"
 #include "Utils/HideIsolateViewModelUtils.h"
 
 using namespace UE::Sequencer;
@@ -1388,9 +1387,7 @@ UWorld* FSequencerFilterBar::GetWorld() const
 
 void FSequencerFilterBar::PreFilter() const
 {
-	TRACE_CPUPROFILER_EVENT_SCOPE(FSequencerFilterBar::PreFilter);
-	
-	const auto PreFilter = [this](const TSharedRef<FSequencerTrackFilter>& InFilter)
+	auto PreFilter = [this](const TSharedRef<FSequencerTrackFilter>& InFilter)
 		{
 			if (IsFilterActive(InFilter))
 			{
@@ -1398,7 +1395,7 @@ void FSequencerFilterBar::PreFilter() const
 			}
 		};
 
-	const auto PreFilterCollection = [PreFilter](const TSharedRef<FSequencerTrackFilterCollection>& InCollection)
+	auto PreFilterCollection = [PreFilter](const TSharedRef<FSequencerTrackFilterCollection>& InCollection)
 		{
 			InCollection->ForEachFilter([PreFilter]
 				(const TSharedRef<FSequencerTrackFilter>& InFilter)
@@ -1419,20 +1416,22 @@ void FSequencerFilterBar::PreFilter() const
 	PreFilter(ModifiedFilter);
 }
 
-// TODO UE-371829: Remove this CVar: Also look for "Sequencer.UseNewFilteringAlgorithm" because it passed to FindConsoleVariable
-static TAutoConsoleVariable<bool> CVarUseNewFilteringAlgorithm(TEXT("Sequencer.UseNewFilteringAlgorithm"), false, TEXT(""));
-
 const FSequencerFilterData& FSequencerFilterBar::FilterNodes()
 {
-	if (CVarUseNewFilteringAlgorithm.GetValueOnAnyThread())
-	{
-		FilterNodesNew();
-	}
-	else
+	//SCOPE_LOG_TIME_IN_SECONDS(TEXT("FSequencerFilterBar::FilterNodes()"), nullptr);
+
+	// Allow all filters to do any pre-processing before filtering
+	PreFilter();
+
+	// Reset all filter data
+	FilterData.Reset();
+
+	// Loop through all nodes and filter recursively
+	const bool bHasActiveFilter = HasAnyFilterActive();
+	for (const TViewModelPtr<IOutlinerExtension>& RootNode : Sequencer.GetNodeTree()->GetRootNodes())
 	{
-		FilterNodesOld();
+		FilterNodesRecursive(bHasActiveFilter, RootNode);
 	}
-	
 
 	// Always filter in spacer node
 	if (const FSequenceModel* const SequenceModel = Sequencer.GetNodeTree()->GetRootNode()->CastThis<FSequenceModel>())
@@ -1502,23 +1501,7 @@ FString FSequencerFilterBar::GenerateTextFilterStringFromEnabledFilters() const
 	return GeneratedFilterString;
 }
 
-void FSequencerFilterBar::FilterNodesOld()
-{
-	// Allow all filters to do any pre-processing before filtering
-	PreFilter();
-	
-	// Reset all filter data
-	FilterData.Reset();
-	
-	// Loop through all nodes and filter recursively
-	const bool bHasActiveFilter = HasAnyFilterActive();
-	for (const TViewModelPtr<IOutlinerExtension>& RootNode : Sequencer.GetNodeTree()->GetRootNodes())
-	{
-		FilterNodesRecursiveOld(bHasActiveFilter, RootNode);
-	}
-}
-
-bool FSequencerFilterBar::FilterNodesRecursiveOld(const bool bInHasActiveFilter, const TViewModelPtr<IOutlinerExtension>& InStartNode)
+bool FSequencerFilterBar::FilterNodesRecursive(const bool bInHasActiveFilter, const TViewModelPtr<IOutlinerExtension>& InStartNode)
 {
 	TRACE_CPUPROFILER_EVENT_SCOPE(FSequencerFilterBar::FilterNodesRecursive);
 	
@@ -1529,163 +1512,62 @@ bool FSequencerFilterBar::FilterNodesRecursiveOld(const bool bInHasActiveFilter,
 	 * - Hidden/Isolated tracks will take precedence over common filters
 	 * - Can hide sub tracks of isolated tracks
 	 */
+
 	bool bAnyChildPassed = false;
-	
+
 	// Child nodes should always be processed, as they may force their parents to pass
 	for (const TViewModelPtr<IOutlinerExtension>& Node : InStartNode.AsModel()->GetChildrenOfType<IOutlinerExtension>())
 	{
-		if (FilterNodesRecursiveOld(bInHasActiveFilter, Node))
+		if (FilterNodesRecursive(bInHasActiveFilter, Node))
 		{
 			bAnyChildPassed = true;
 		}
 	}
-	
+
 	// Increment the total node count so we can remove the code to loop again just to count
 	FilterData.IncrementTotalNodeCount();
+
 	// Early out if no filter
 	if (!bInHasActiveFilter)
 	{
 		FilterData.FilterInNode(InStartNode);
 		return false;
 	}
-	
+
 	const USequencerSettings* const SequencerSettings = Sequencer.GetSequencerSettings();
 	if (!SequencerSettings)
 	{
 		return false;
 	}
-	
+
 	// Pinning overrides all other filters
 	if (EnumHasAnyFlags(Flags, EFilterFlags::PinnedItemsCanSkipFiltering) && !SequencerSettings->GetIncludePinnedInFilter())
 	{
 		const TSharedPtr<IPinnableExtension> Pinnable = InStartNode.AsModel()->FindAncestorOfType<IPinnableExtension>(true);
 		if (Pinnable.IsValid() && Pinnable->IsPinned())
 		{
-			FilterData.FilterInParentChildNodes(InStartNode, ENodeVisitFlags::IncludeAll);
+			FilterData.FilterInParentChildNodes(InStartNode, true, true, true);
 			return true;
 		}
 	}
-	
+
 	const bool bPassedTextFilter = !TextFilter->IsActive() || TextFilter->PassesFilter(InStartNode);
 	const bool bPassedHideIsolateFilter = !HideIsolateFilter->IsActive() || HideIsolateFilter->PassesFilter(InStartNode);
 	const bool bPassedAnyCommonFilters = PassesAnyCommonFilter(InStartNode);
 	const bool bPassedInternalFilters = !AnyInternalFilterActive() || PassesAllInternalFilters(InStartNode);
 	const bool bPassedAnyCustomTextFilters = PassesAllCustomTextFilters(InStartNode);
+
 	const bool bAllFiltersPassed = bPassedTextFilter
 		&& bPassedHideIsolateFilter
 		&& bPassedAnyCommonFilters
 		&& bPassedInternalFilters
 		&& bPassedAnyCustomTextFilters;
-	
-	if (bAllFiltersPassed || bAnyChildPassed)
-	{
-		// Don't expand if a mouse capture is effecting TrackValueChanged events to avoid keys
-		// being hidden and expanding tracks while dragging
-		if (SequencerSettings->GetAutoExpandNodesOnFilterPass()
-			&& !FSlateApplication::Get().HasAnyMouseCaptor())
-		{
-			SetTrackParentsExpanded(InStartNode.ImplicitCast(), true);
-		}
-		FilterData.FilterInNodeWithAncestors(InStartNode);
-		return true;
-	}
-	// After child nodes are processed, fail anything that didn't pass
-	FilterData.FilterOutNode(InStartNode);
-	return false;
-}
 
-void FSequencerFilterBar::FilterNodesNew()
-{
-	// Reset all filter data
-	FilterData.Reset();
-	
-	if (HasAnyFilterActive())
-	{
-		const TArray<TNotNull<FSequencerTrackFilter*>> Filters = BuildActiveFilters();
-		
-		for (FSequencerTrackFilter* Filter : Filters)
-		{
-			Filter->OnPreFilter();
-		}
-		
-		// Loop through all nodes and filter recursively
-		for (const TViewModelPtr<IOutlinerExtension>& RootNode : Sequencer.GetNodeTree()->GetRootNodes())
-		{
-			FilterNodesRecursiveNew(RootNode, Filters, {});
-		}
-	}
-	else
-	{
-		for (const TViewModelPtr<IOutlinerExtension>& RootNode : Sequencer.GetNodeTree()->GetRootNodes())
-		{
-			FilterData.FilterInParentChildNodes(RootNode, ENodeVisitFlags::IncludeSelf | ENodeVisitFlags::IncludeChildren);
-		}
-	}
-}
-
-bool FSequencerFilterBar::FilterNodesRecursiveNew(
-	const TViewModelPtr<IOutlinerExtension>& InStartNode,
-	TConstArrayView<TNotNull<FSequencerTrackFilter*>> InFilters,
-	const FRecursiveFilterOverrides& InFilterOverrides
-	)
-{
-	TRACE_CPUPROFILER_EVENT_SCOPE(FSequencerFilterBar::FilterNodesRecursive);
-	
-	/**
-	 * Main Filtering Logic
-	 *
-	 * - Pinning overrides all other filters
-	 * - Hidden/Isolated tracks will take precedence over common filters
-	 * - Can hide sub tracks of isolated tracks
-	 */
-
-	const USequencerSettings* const SequencerSettings = Sequencer.GetSequencerSettings();
-
-	// Pinning overrides all other filters
-	if (EnumHasAnyFlags(Flags, EFilterFlags::PinnedItemsCanSkipFiltering) && SequencerSettings && !SequencerSettings->GetIncludePinnedInFilter())
-	{
-		const TSharedPtr<IPinnableExtension> Pinnable = InStartNode.AsModel()->FindAncestorOfType<IPinnableExtension>(true);
-		if (Pinnable.IsValid() && Pinnable->IsPinned())
-		{
-			const ENodeVisitFlags VisitFlags = ENodeVisitFlags::IncludeSelf | ENodeVisitFlags::IncludeChildren;
-			FilterData.FilterInParentChildNodes(InStartNode, VisitFlags);
-			return true;
-		}
-	}
-	
-	const FFilterEvaluation EvaluationResult = EvaluateFiltersOnItem(InStartNode, InFilters, InFilterOverrides);
-	const bool bAllFiltersPassed = EvaluationResult.ItemResult != EItemFilterState::Exclude;
-	
-	bool bAnyChildPassed = false;
-	const bool bHasOverrideForAllFilters = EvaluationResult.NewRecursiveOverrides.Overrides.Num() == InFilters.Num();
-	if (bHasOverrideForAllFilters)
-	{
-		bAnyChildPassed = bAllFiltersPassed;
-		const ENodeVisitFlags VisitFlags = ENodeVisitFlags::IncludeSelf | ENodeVisitFlags::IncludeChildren;
-		const int32 Num = bAllFiltersPassed 
-			? FilterData.FilterInParentChildNodes(InStartNode, VisitFlags) 
-			: FilterData.FilterOutParentChildNodes(InStartNode, VisitFlags);
-		FilterData.AddTotalNodeCount(Num);
-	}
-	else
-	{
-		FilterData.IncrementTotalNodeCount();
-		
-		// Child nodes should always be processed, as they may force their parents to pass
-		for (const TViewModelPtr<IOutlinerExtension>& Node : InStartNode.AsModel()->GetChildrenOfType<IOutlinerExtension>())
-		{
-			if (FilterNodesRecursiveNew(Node, InFilters, EvaluationResult.NewRecursiveOverrides))
-			{
-				bAnyChildPassed = true;
-			}
-		}
-	}
-	
 	if (bAllFiltersPassed || bAnyChildPassed)
 	{
 		// Don't expand if a mouse capture is effecting TrackValueChanged events to avoid keys
 		// being hidden and expanding tracks while dragging
-		if (SequencerSettings && SequencerSettings->GetAutoExpandNodesOnFilterPass()
+		if (SequencerSettings->GetAutoExpandNodesOnFilterPass()
 			&& !FSlateApplication::Get().HasAnyMouseCaptor())
 		{
 			SetTrackParentsExpanded(InStartNode.ImplicitCast(), true);
@@ -1700,49 +1582,6 @@ bool FSequencerFilterBar::FilterNodesRecursiveNew(
 	return false;
 }
 
-TArray<TNotNull<FSequencerTrackFilter*>> FSequencerFilterBar::BuildActiveFilters()
-{
-	TArray<TNotNull<FSequencerTrackFilter*>> Filters;
-	
-	if (TextFilter->IsActive())
-	{
-		Filters.Emplace(&TextFilter.Get());
-	}
-	
-	if (HideIsolateFilter->IsActive())
-	{
-		Filters.Emplace(&HideIsolateFilter.Get());
-	}
-	
-	CommonFilters->ForEachFilter([this, &Filters](const TSharedRef<FSequencerTrackFilter>& InFilter)
-	{
-		if (IsFilterActive(InFilter))
-		{
-			Filters.Emplace(&InFilter.Get());
-		}
-		return true;
-	});
-	
-	InternalFilters->ForEachFilter([this, &Filters](const TSharedRef<FSequencerTrackFilter>& InFilter)
-	{
-		if (IsFilterActive(InFilter))
-		{
-			Filters.Emplace(&InFilter.Get());
-		}
-		return true;
-	});
-	
-	for (const TSharedRef<FSequencerTrackFilter_CustomText>& Filter : CustomTextFilters)
-	{
-		if (IsFilterActive(Filter))
-		{
-			Filters.Emplace(&Filter.Get());
-		}
-	}
-	
-	return Filters;
-}
-
 void FSequencerFilterBar::HandleTextFilterChanged()
 {
 	OnTextFilterTextChangedEvent.Broadcast();
diff --git a/Engine/Source/Editor/Sequencer/Private/Filters/SequencerFilterBar.h b/Engine/Source/Editor/Sequencer/Private/Filters/SequencerFilterBar.h
--- a/Engine/Source/Editor/Sequencer/Private/Filters/SequencerFilterBar.h
+++ b/Engine/Source/Editor/Sequencer/Private/Filters/SequencerFilterBar.h
@@ -26,9 +26,7 @@ class SFilterBarIsolateHideShow;
 class SSequencerFilterBar;
 class SSequencerSearchBox;
 enum class EFilterBarLayout : uint8;
-namespace UE::Sequencer {
-struct FRecursiveFilterOverrides;
-class IOutlinerExtension; }
+namespace UE::Sequencer { class IOutlinerExtension; }
 namespace UE::Sequencer { struct FFilterBarInitArgs; }
 
 enum class ESequencerFilterChange : uint8
@@ -323,21 +321,8 @@ private:
 	/** Binds all commands that the filter bar or its filters handle. */
 	void BindCommands();
 	
-	/** Old algorithm for filtering nodes */
-	void FilterNodesOld();
 	/** Do not call directly! Should only be called by FilterNodes(). */
-	bool FilterNodesRecursiveOld(const bool bInHasActiveFilter, const UE::Sequencer::TViewModelPtr<UE::Sequencer::IOutlinerExtension>& InStartNode);
-	
-	/** New algorithm for filtering nodes. */
-	void FilterNodesNew();
-	/** @return Filter that are active and should be evaluated. */
-	TArray<TNotNull<FSequencerTrackFilter*>> BuildActiveFilters();
-	/** Do not call directly! Should only be called by FilterNodes(). */
-	bool FilterNodesRecursiveNew(
-		const UE::Sequencer::TViewModelPtr<UE::Sequencer::IOutlinerExtension>& InStartNode,
-		TConstArrayView<TNotNull<FSequencerTrackFilter*>> InFilters,
-		const UE::Sequencer::FRecursiveFilterOverrides& InFilterOverrides
-		);
+	bool FilterNodesRecursive(const bool bInHasActiveFilter, const UE::Sequencer::TViewModelPtr<UE::Sequencer::IOutlinerExtension>& InStartNode);
 	
 	/** Invoked when the text filter changes. */
 	void HandleTextFilterChanged();
diff --git a/Engine/Source/Editor/Sequencer/Private/Filters/SequencerFilterData.cpp b/Engine/Source/Editor/Sequencer/Private/Filters/SequencerFilterData.cpp
--- a/Engine/Source/Editor/Sequencer/Private/Filters/SequencerFilterData.cpp
+++ b/Engine/Source/Editor/Sequencer/Private/Filters/SequencerFilterData.cpp
@@ -63,11 +63,6 @@ void FSequencerFilterData::IncrementTotalNodeCount()
 	++TotalNodeCount;
 }
 
-void FSequencerFilterData::AddTotalNodeCount(int32 InNodeCount)
-{
-	TotalNodeCount += FMath::Max(0, InNodeCount);
-}
-
 void FSequencerFilterData::FilterInNode(TWeakViewModelPtr<IOutlinerExtension> InNodeWeak)
 {
 	FilterInNodes.Add(InNodeWeak);
@@ -92,32 +87,39 @@ void FSequencerFilterData::FilterOutNode(TWeakViewModelPtr<IOutlinerExtension> I
 	}
 }
 
-void FSequencerFilterData::FilterInParentChildNodes(
-	const TViewModelPtr<IOutlinerExtension>& InNode, const bool bInIncludeSelf, const bool bInIncludeParents, const bool bInIncludeChildren
-	)
+void FSequencerFilterData::FilterInParentChildNodes(const TViewModelPtr<IOutlinerExtension>& InNode
+	, const bool bInIncludeSelf, const bool bInIncludeParents, const bool bInIncludeChildren)
 {
-	const ENodeVisitFlags Flags = 
-		(bInIncludeSelf ? ENodeVisitFlags::IncludeSelf : ENodeVisitFlags::None)
-		| (bInIncludeParents ? ENodeVisitFlags::IncludeParents : ENodeVisitFlags::None)
-		| (bInIncludeChildren ? ENodeVisitFlags::IncludeChildren : ENodeVisitFlags::None);
-	FilterInParentChildNodes(InNode, Flags);
-}
+	if (!InNode.IsValid())
+	{
+		return;
+	}
 
-int32 FSequencerFilterData::FilterInParentChildNodes(const TViewModelPtr<IOutlinerExtension>& InNode, ENodeVisitFlags InFlags)
-{
-	const auto FilterInLambda = [this](const TViewModelPtr<IOutlinerExtension>& Node){ FilterInNode(Node); };
-	return ApplyToNodes(InNode, FilterInLambda, InFlags);
-}
+	if (bInIncludeParents)
+	{
+		for (TViewModelPtr<IOutlinerExtension> ParentNode : InNode.AsModel()->GetAncestorsOfType<IOutlinerExtension>())
+		{
+			FilterInNode(ParentNode);
+		}
+	}
 
-int32 FSequencerFilterData::FilterOutParentChildNodes(const TViewModelPtr<IOutlinerExtension>& InNode, ENodeVisitFlags InFlags)
-{
-	const auto FilterOutLambda = [this](const TViewModelPtr<IOutlinerExtension>& Node){ FilterOutNode(Node); };
-	return ApplyToNodes(InNode, FilterOutLambda, InFlags);
+	if (bInIncludeSelf)
+	{
+		FilterInNode(InNode);
+	}
+
+	if (bInIncludeChildren)
+	{
+		for (TViewModelPtr<IOutlinerExtension> ChildNode : InNode.AsModel()->GetDescendantsOfType<IOutlinerExtension>())
+		{
+			FilterInNode(ChildNode);
+		}
+	}
 }
 
 void FSequencerFilterData::FilterInNodeWithAncestors(const TViewModelPtr<IOutlinerExtension>& InNode)
 {
-	FilterInParentChildNodes(InNode, ENodeVisitFlags::IncludeSelf | ENodeVisitFlags::IncludeParents);
+	FilterInParentChildNodes(InNode, true, true, false);
 }
 
 bool FSequencerFilterData::ContainsFilterInNodes(const FSequencerFilterData& InOtherData) const
@@ -223,42 +225,3 @@ UObject* FSequencerFilterData::ResolveTrackBoundObject(ISequencer& InSequencer,
 
 	return nullptr;
 }
-
-int32 FSequencerFilterData::ApplyToNodes(
-	const TViewModelPtr<IOutlinerExtension>& InNode, 
-	TFunctionRef<void(const TViewModelPtr<IOutlinerExtension>&)>&& InCallback,
-	ENodeVisitFlags InFlags
-	)
-{
-	int32 NodeCount = 0;
-	if (!InNode.IsValid())
-	{
-		return NodeCount;
-	}
-	
-	if (EnumHasAnyFlags(InFlags, ENodeVisitFlags::IncludeParents))
-	{
-		for (TViewModelPtr<IOutlinerExtension> ParentNode : InNode.AsModel()->GetAncestorsOfType<IOutlinerExtension>())
-		{
-			++NodeCount;
-			InCallback(ParentNode);
-		}
-	}
-
-	if (EnumHasAnyFlags(InFlags, ENodeVisitFlags::IncludeSelf))
-	{
-		++NodeCount;
-		InCallback(InNode);
-	}
-
-	if (EnumHasAnyFlags(InFlags, ENodeVisitFlags::IncludeChildren))
-	{
-		for (TViewModelPtr<IOutlinerExtension> ChildNode : InNode.AsModel()->GetDescendantsOfType<IOutlinerExtension>())
-		{
-			++NodeCount;
-			InCallback(ChildNode);
-		}
-	}
-	
-	return NodeCount;
-}
\ No newline at end of file
diff --git a/Engine/Source/Editor/Sequencer/Private/Filters/Utils/FilterEvaluation.cpp b/Engine/Source/Editor/Sequencer/Private/Filters/Utils/FilterEvaluation.cpp
deleted file mode 100644
diff --git a/Engine/Source/Editor/Sequencer/Private/Filters/Utils/FilterEvaluation.h b/Engine/Source/Editor/Sequencer/Private/Filters/Utils/FilterEvaluation.h
deleted file mode 100644
diff --git a/Engine/Source/Editor/Sequencer/Public/Filters/FilterResult.h b/Engine/Source/Editor/Sequencer/Public/Filters/FilterResult.h
deleted file mode 100644
diff --git a/Engine/Source/Editor/Sequencer/Public/Filters/Filters/SequencerTrackFilter_Selected.h b/Engine/Source/Editor/Sequencer/Public/Filters/Filters/SequencerTrackFilter_Selected.h
--- a/Engine/Source/Editor/Sequencer/Public/Filters/Filters/SequencerTrackFilter_Selected.h
+++ b/Engine/Source/Editor/Sequencer/Public/Filters/Filters/SequencerTrackFilter_Selected.h
@@ -19,7 +19,6 @@ public:
 	virtual FText GetDefaultToolTipText() const override;
 	virtual TSharedPtr<FUICommandInfo> GetToggleCommand() const override;
 	virtual void ActiveStateChanged(const bool bInActive) override;
-	virtual UE::Sequencer::FFilterResult Evaluate(FSequencerTrackFilterType InItem) const override;
 	//~ End FSequencerTrackFilter
 
 	//~ Begin FFilterBase
@@ -27,34 +26,23 @@ public:
 	virtual FSlateIcon GetIcon() const override;
 	//~ End FFilterBase
 
-	virtual void OnPreFilter() override;
-
 	//~ Begin IFilter
 	virtual FString GetName() const override;
 	virtual bool PassesFilter(FSequencerTrackFilterType InItem) const override;
 	//~ End IFilter
 
+	void ToggleShowOnlySelectedTracks();
+
 protected:
-	
 	void BindSelectionChanged();
 	void UnbindSelectionChanged();
 
+	void OnSelectionChanged(UObject* const InObject);
+
 	FDelegateHandle OnSelectionChangedHandle;
 	
 private:
 	
 	/** Objects extending this filter. Cached at construction to speed up filtering. */
 	const TArray<TWeakObjectPtr<USequencerTrackFilterSelectedExtension>> CachedWeakExtensions;
-	
-	/** Whether IncludedViewModels needs to be rebuilt. */
-	bool bNeedsPrefilterRebuild = true;
-	
-	/** The view models that pass the filter. Built once during pre-filter to avoid excessive iteration of tree hierarchy. */
-	TMap<UE::Sequencer::TWeakViewModelPtr<UE::Sequencer::FViewModel>, UE::Sequencer::EFilterResultFlags> IncludedViewModels;
-	
-	void BindEvents();
-	void UnbindEvents();
-	
-	void OnSequencerSelectionChanged();
-	void OnViewportSelectionChanged(UObject* const InObject);
 };
diff --git a/Engine/Source/Editor/Sequencer/Public/Filters/SequencerFilterData.h b/Engine/Source/Editor/Sequencer/Public/Filters/SequencerFilterData.h
--- a/Engine/Source/Editor/Sequencer/Public/Filters/SequencerFilterData.h
+++ b/Engine/Source/Editor/Sequencer/Public/Filters/SequencerFilterData.h
@@ -3,9 +3,8 @@
 #pragma once
 
 #include "FilterBarDelegates.h"
-#include "MVVM/ViewModelPtr.h"
-#include "Misc/EnumClassFlags.h"
 #include "MovieSceneTrack.h"
+#include "MVVM/ViewModelPtr.h"
 
 #define UE_API SEQUENCER_API
 
@@ -20,22 +19,6 @@ namespace UE::Sequencer
 
 using FSequencerTrackFilterType = UE::Sequencer::FViewModelPtr;
 
-namespace UE::Sequencer
-{
-/** Utility argument for FSequencerFilterData. */
-enum class ENodeVisitFlags : uint8
-{
-	None,
-	
-	IncludeSelf = 1 << 0,
-	IncludeParents = 1 << 1,
-	IncludeChildren = 1 << 2,
-	
-	IncludeAll = IncludeSelf | IncludeParents | IncludeChildren,
-};
-ENUM_CLASS_FLAGS(ENodeVisitFlags);
-}
-
 /** Represents a cache between nodes for a filter operation. */
 struct FSequencerFilterData
 {
@@ -55,23 +38,12 @@ struct FSequencerFilterData
 	UE_API uint32 GetFilterOutCount() const;
 
 	UE_API void IncrementTotalNodeCount();
-	UE_API void AddTotalNodeCount(int32 InNodeCount);
 
 	UE_API void FilterInNode(UE::Sequencer::TWeakViewModelPtr<UE::Sequencer::IOutlinerExtension> InNodeWeak);
 	UE_API void FilterOutNode(UE::Sequencer::TWeakViewModelPtr<UE::Sequencer::IOutlinerExtension> InNodeWeak);
 
-	UE_DEPRECATED(5.8, "Use the version with EFilterInFlags instead")
-	UE_API void FilterInParentChildNodes(const UE::Sequencer::TViewModelPtr<UE::Sequencer::IOutlinerExtension>& InNode, const bool bInIncludeSelf, const bool bInIncludeParents, const bool bInIncludeChildren = false);
-	
-	/** @return Number of nodes visited */
-	UE_API int32 FilterInParentChildNodes(
-		const UE::Sequencer::TViewModelPtr<UE::Sequencer::IOutlinerExtension>& InNode, UE::Sequencer::ENodeVisitFlags InFlags
-		);
-	/** @return Number of nodes visited */
-	UE_API int32 FilterOutParentChildNodes(
-		const UE::Sequencer::TViewModelPtr<UE::Sequencer::IOutlinerExtension>& InNode, 
-		UE::Sequencer::ENodeVisitFlags InFlags
-		);
+	UE_API void FilterInParentChildNodes(const UE::Sequencer::TViewModelPtr<UE::Sequencer::IOutlinerExtension>& InNode
+		, const bool bInIncludeSelf, const bool bInIncludeParents, const bool bInIncludeChildren = false);
 
 	UE_API void FilterInNodeWithAncestors(const UE::Sequencer::TViewModelPtr<UE::Sequencer::IOutlinerExtension>& InNode);
 
@@ -100,13 +72,6 @@ protected:
 	
 	/** Broadcasts when the filter state of IOutlinerExtension is changed. */
 	UE::Sequencer::FChangeOutlinerExtensionFilterState ChangeFilterStateEvent;
-	
-	/** @return Number of nodes visited */
-	int32 ApplyToNodes(
-		const UE::Sequencer::TViewModelPtr<UE::Sequencer::IOutlinerExtension>& InNode, 
-		TFunctionRef<void(const UE::Sequencer::TViewModelPtr<UE::Sequencer::IOutlinerExtension>&)>&& InCallback,
-		UE::Sequencer::ENodeVisitFlags InFlags
-		);
 };
 
 #undef UE_API
diff --git a/Engine/Source/Editor/Sequencer/Public/Filters/SequencerTrackFilterBase.h b/Engine/Source/Editor/Sequencer/Public/Filters/SequencerTrackFilterBase.h
--- a/Engine/Source/Editor/Sequencer/Public/Filters/SequencerTrackFilterBase.h
+++ b/Engine/Source/Editor/Sequencer/Public/Filters/SequencerTrackFilterBase.h
@@ -2,7 +2,6 @@
 
 #pragma once
 
-#include "FilterResult.h"
 #include "Filters/FilterBase.h"
 #include "Filters/ISequencerTrackFilters.h"
 #include "GameFramework/Actor.h"
@@ -50,12 +49,11 @@ public:
 	SEQUENCER_API FSequencerTrackFilter(ISequencerTrackFilters& InOutFilterInterface, TSharedPtr<FFilterCategory>&& InCategory = nullptr);
 
 	//~ Begin FSequencerFilterBase
+
 	SEQUENCER_API virtual ISequencerTrackFilters& GetFilterInterface() const override;
+
 	//~ End FSequencerFilterBase
 
-	/** Evaluates the filter on an item. It allows returning more information to optimize evaluation. */
-	SEQUENCER_API virtual UE::Sequencer::FFilterResult Evaluate(FSequencerTrackFilterType InItem) const;
-	
 	/** Returns whether this filter needs reevaluating any time track values have been modified, not just tree changes */
 	virtual bool ShouldUpdateOnTrackValueChanged() const { return false; }
 
