Sequencer: Fix new filtering algorithm evaluation to be a disjunctive, like it was in the old algorithm. Compare to FSequencerFilterBar::FilterNodesRecursiveOld::bAllFiltersPassed

For example, if the user has the Keyed and Camera fitler active, each adds to the fitler list. You don't want items that are Camera tracks AND Keyed, you want items that are either Camera tracks OR keyed.

#Synced-CL 52283655

diff --git a/Engine/Source/Editor/Sequencer/Private/Filters/SequencerFilterBar.cpp b/Engine/Source/Editor/Sequencer/Private/Filters/SequencerFilterBar.cpp
--- a/Engine/Source/Editor/Sequencer/Private/Filters/SequencerFilterBar.cpp
+++ b/Engine/Source/Editor/Sequencer/Private/Filters/SequencerFilterBar.cpp
@@ -33,8 +33,8 @@
 #include "Filters/SequencerFilterAreaConfig.h"
 #include "Stats/StatsMisc.h"
 #include "UObject/UObjectIterator.h"
-#include "Utils/FilterEvaluation.h"
 #include "Utils/HideIsolateViewModelUtils.h"
+#include "Utils/Evaluation/FilterBarEvaluation.h"
 
 using namespace UE::Sequencer;
 
@@ -1420,7 +1420,7 @@ void FSequencerFilterBar::PreFilter() const
 }
 
 // TODO UE-371829: Remove this CVar: Also look for "Sequencer.UseNewFilteringAlgorithm" because it passed to FindConsoleVariable
-static TAutoConsoleVariable<bool> CVarUseNewFilteringAlgorithm(TEXT("Sequencer.UseNewFilteringAlgorithm"), false, TEXT(""));
+static TAutoConsoleVariable<bool> CVarUseNewFilteringAlgorithm(TEXT("Sequencer.UseNewFilteringAlgorithm"), true, TEXT(""));
 
 const FSequencerFilterData& FSequencerFilterBar::FilterNodes()
 {
@@ -1601,12 +1601,8 @@ void FSequencerFilterBar::FilterNodesNew()
 	
 	if (HasAnyFilterActive())
 	{
-		const TArray<TNotNull<FSequencerTrackFilter*>> Filters = BuildActiveFilters();
-		
-		for (FSequencerTrackFilter* Filter : Filters)
-		{
-			Filter->OnPreFilter();
-		}
+		const FSequencerFilterContext Filters = BuildActiveFilters();
+		InvokePreFilter(Filters);
 		
 		// Loop through all nodes and filter recursively
 		for (const TViewModelPtr<IOutlinerExtension>& RootNode : Sequencer.GetNodeTree()->GetRootNodes())
@@ -1625,8 +1621,8 @@ void FSequencerFilterBar::FilterNodesNew()
 
 bool FSequencerFilterBar::FilterNodesRecursiveNew(
 	const TViewModelPtr<IOutlinerExtension>& InStartNode,
-	TConstArrayView<TNotNull<FSequencerTrackFilter*>> InFilters,
-	const FRecursiveFilterOverrides& InFilterOverrides
+	const FSequencerFilterContext& InFilters,
+	const FSequencerFilterOverrides& InOverrides
 	)
 {
 	TRACE_CPUPROFILER_EVENT_SCOPE(FSequencerFilterBar::FilterNodesRecursive);
@@ -1653,11 +1649,12 @@ bool FSequencerFilterBar::FilterNodesRecursiveNew(
 		}
 	}
 	
-	const FFilterEvaluation EvaluationResult = EvaluateFiltersOnItem(InStartNode, InFilters, InFilterOverrides);
-	const bool bAllFiltersPassed = EvaluationResult.ItemResult != EItemFilterState::Exclude;
-	
+	const FSequencerFilterEvaluation EvaluationResult = EvaluateSequencerFilterBar(InStartNode, InFilters, InOverrides);
+	const bool bAllFiltersPassed = PassesFilterState(EvaluationResult.ItemResult);
 	bool bAnyChildPassed = false;
-	const bool bHasOverrideForAllFilters = EvaluationResult.NewRecursiveOverrides.Overrides.Num() == InFilters.Num();
+	
+	// Some filters, such as text filters, need to be run on all children. If we're not evaluating text filters, we can optimize by trimming the tree.
+	const bool bHasOverrideForAllFilters = InFilters.HasOnlyCachedFilters() && EvaluationResult.NewOverrides.HasOverrideForAllFilters(InFilters);
 	if (bHasOverrideForAllFilters)
 	{
 		bAnyChildPassed = bAllFiltersPassed;
@@ -1674,7 +1671,7 @@ bool FSequencerFilterBar::FilterNodesRecursiveNew(
 		// Child nodes should always be processed, as they may force their parents to pass
 		for (const TViewModelPtr<IOutlinerExtension>& Node : InStartNode.AsModel()->GetChildrenOfType<IOutlinerExtension>())
 		{
-			if (FilterNodesRecursiveNew(Node, InFilters, EvaluationResult.NewRecursiveOverrides))
+			if (FilterNodesRecursiveNew(Node, InFilters, EvaluationResult.NewOverrides))
 			{
 				bAnyChildPassed = true;
 			}
@@ -1700,25 +1697,25 @@ bool FSequencerFilterBar::FilterNodesRecursiveNew(
 	return false;
 }
 
-TArray<TNotNull<FSequencerTrackFilter*>> FSequencerFilterBar::BuildActiveFilters()
+FSequencerFilterContext FSequencerFilterBar::BuildActiveFilters()
 {
-	TArray<TNotNull<FSequencerTrackFilter*>> Filters;
+	FSequencerFilterContext Filters;
 	
 	if (TextFilter->IsActive())
 	{
-		Filters.Emplace(&TextFilter.Get());
+		Filters.TextFilter = &TextFilter.Get();
 	}
 	
 	if (HideIsolateFilter->IsActive())
 	{
-		Filters.Emplace(&HideIsolateFilter.Get());
+		Filters.HideIsolate = &HideIsolateFilter.Get();
 	}
 	
 	CommonFilters->ForEachFilter([this, &Filters](const TSharedRef<FSequencerTrackFilter>& InFilter)
 	{
 		if (IsFilterActive(InFilter))
 		{
-			Filters.Emplace(&InFilter.Get());
+			Filters.CommonFilters.Add(&InFilter.Get());
 		}
 		return true;
 	});
@@ -1727,7 +1724,7 @@ TArray<TNotNull<FSequencerTrackFilter*>> FSequencerFilterBar::BuildActiveFilters
 	{
 		if (IsFilterActive(InFilter))
 		{
-			Filters.Emplace(&InFilter.Get());
+			Filters.InternalFilters.Add(&InFilter.Get());
 		}
 		return true;
 	});
@@ -1736,7 +1733,7 @@ TArray<TNotNull<FSequencerTrackFilter*>> FSequencerFilterBar::BuildActiveFilters
 	{
 		if (IsFilterActive(Filter))
 		{
-			Filters.Emplace(&Filter.Get());
+			Filters.CustomTextFilters.Add(&Filter.Get());
 		}
 	}
 	
diff --git a/Engine/Source/Editor/Sequencer/Private/Filters/SequencerFilterBar.h b/Engine/Source/Editor/Sequencer/Private/Filters/SequencerFilterBar.h
--- a/Engine/Source/Editor/Sequencer/Private/Filters/SequencerFilterBar.h
+++ b/Engine/Source/Editor/Sequencer/Private/Filters/SequencerFilterBar.h
@@ -27,6 +27,8 @@ class SSequencerFilterBar;
 class SSequencerSearchBox;
 enum class EFilterBarLayout : uint8;
 namespace UE::Sequencer {
+struct FSequencerFilterOverrides;
+struct FSequencerFilterContext;
 struct FRecursiveFilterOverrides;
 class IOutlinerExtension; }
 namespace UE::Sequencer { struct FFilterBarInitArgs; }
@@ -331,12 +333,12 @@ private:
 	/** New algorithm for filtering nodes. */
 	void FilterNodesNew();
 	/** @return Filter that are active and should be evaluated. */
-	TArray<TNotNull<FSequencerTrackFilter*>> BuildActiveFilters();
+	UE::Sequencer::FSequencerFilterContext BuildActiveFilters();
 	/** Do not call directly! Should only be called by FilterNodes(). */
 	bool FilterNodesRecursiveNew(
 		const UE::Sequencer::TViewModelPtr<UE::Sequencer::IOutlinerExtension>& InStartNode,
-		TConstArrayView<TNotNull<FSequencerTrackFilter*>> InFilters,
-		const UE::Sequencer::FRecursiveFilterOverrides& InFilterOverrides
+		const UE::Sequencer::FSequencerFilterContext& InFilters,
+		const UE::Sequencer::FSequencerFilterOverrides& InOverrides
 		);
 	
 	/** Invoked when the text filter changes. */
diff --git a/Engine/Source/Editor/Sequencer/Private/Filters/Utils/Evaluation/FilterBarEvaluation.cpp b/Engine/Source/Editor/Sequencer/Private/Filters/Utils/Evaluation/FilterBarEvaluation.cpp
new file mode 100644
--- /dev/null
+++ b/Engine/Source/Editor/Sequencer/Private/Filters/Utils/Evaluation/FilterBarEvaluation.cpp
@@ -0,0 +1,108 @@
+﻿// Copyright Epic Games, Inc. All Rights Reserved.
+
+#include "FilterBarEvaluation.h"
+
+#include "Filters/SequencerTrackFilterBase.h"
+#include "Filters/Filters/SequencerTrackFilter_HideIsolate.h"
+#include "Filters/Filters/SequencerTrackFilter_Text.h"
+
+namespace UE::Sequencer
+{
+void InvokePreFilter(const FSequencerFilterContext& InFilters)
+{
+	// Text filters, and isolte/hide/show do not require pre-filter.
+	
+	for (const TNotNull<FSequencerTrackFilter*> CommonFilter : InFilters.CommonFilters)
+	{
+		CommonFilter->OnPreFilter();
+	}
+	
+	for (const TNotNull<FSequencerTrackFilter*> InternalFilter : InFilters.InternalFilters)
+	{
+		InternalFilter->OnPreFilter();
+	}
+}
+
+namespace FilterEval
+{
+template<typename TFilter> requires std::is_base_of_v<FSequencerTrackFilter, TFilter>
+static EItemFilterState EvaluateFilter(
+	const TViewModelPtr<IOutlinerExtension>& InItem, const TOptional<TNotNull<TFilter*>> InFilter
+	)
+{
+	if (!InFilter)
+	{
+		return EItemFilterState::DoNotCare;
+	}
+	
+	return InFilter.GetValue()->PassesFilter(InItem) ? EItemFilterState::Include : EItemFilterState::Exclude;
+}
+
+static EItemFilterState EvaluateFilters(
+	const TViewModelPtr<IOutlinerExtension>& InItem, const TConstArrayView<TNotNull<FSequencerTrackFilter*>> InFilters
+	)
+{
+	if (InFilters.Num() == 0)
+	{
+		return EItemFilterState::DoNotCare;
+	}
+	
+	const bool bPassedAll = Algo::AllOf(InFilters, [&InItem](TNotNull<FSequencerTrackFilter*> CustomTextFilter)
+	{
+		return CustomTextFilter->PassesFilter(InItem);
+	});
+	return bPassedAll? EItemFilterState::Include : EItemFilterState::Exclude;
+}
+}
+
+FSequencerFilterEvaluation EvaluateSequencerFilterBar(
+	const TViewModelPtr<IOutlinerExtension>& InItem, 
+	const FSequencerFilterContext& InFilters,
+	const FSequencerFilterOverrides& InOverrides
+	)
+{
+	FSequencerFilterOverrides NewOverrides = InOverrides;
+	
+	// The below is sorted in the order we expect to evaluate the fastest.
+	
+	const EItemFilterState HideIsolateFilterResult = FilterEval::EvaluateFilter(InItem, InFilters.HideIsolate);
+	if (!PassesFilterState(HideIsolateFilterResult))
+	{
+		return FSequencerFilterEvaluation(HideIsolateFilterResult, NewOverrides);
+	}
+	
+	const EItemFilterState TextFilterResult = FilterEval::EvaluateFilter(InItem, InFilters.TextFilter);
+	if (!PassesFilterState(TextFilterResult))
+	{
+		return FSequencerFilterEvaluation(TextFilterResult, NewOverrides);
+	}
+	
+	const EItemFilterState CustomTextResult = FilterEval::EvaluateFilters(InItem, InFilters.CustomTextFilters);
+	if (!PassesFilterState(CustomTextResult))
+	{
+		return FSequencerFilterEvaluation(CustomTextResult, NewOverrides);
+	}
+	
+	// The internal filters act as a pre-filter... before user filters are applied.
+	const FFilterGroupEvaluation InternalEval = EvaluateConjunctive(InItem, InFilters.InternalFilters, InOverrides.InternalOverrides);
+	NewOverrides.InternalOverrides = InternalEval.NewRecursiveOverrides;
+	if (!PassesFilterState(InternalEval.ItemResult))
+	{
+		return FSequencerFilterEvaluation(InternalEval.ItemResult, NewOverrides);
+	}
+	
+	// The filter pills that the user manually checks and unchecks add or remove items... which encodes an OR condition
+	const FFilterGroupEvaluation CommonEval = EvaluateDisjunctive(InItem, InFilters.CommonFilters, InOverrides.CommonOverrides);
+	NewOverrides.CommonOverrides = CommonEval.NewRecursiveOverrides;
+	if (!PassesFilterState(CommonEval.ItemResult))
+	{
+		return FSequencerFilterEvaluation(CommonEval.ItemResult, NewOverrides);
+	}
+	
+	// All filter results must pass
+	const EItemFilterState CombinedResult = CombineConjunctive(
+		TextFilterResult, HideIsolateFilterResult, CommonEval.ItemResult, InternalEval.ItemResult, CustomTextResult
+		);
+	return FSequencerFilterEvaluation(CombinedResult, NewOverrides);
+}
+} // namespace UE::Sequencer
diff --git a/Engine/Source/Editor/Sequencer/Private/Filters/Utils/Evaluation/FilterBarEvaluation.h b/Engine/Source/Editor/Sequencer/Private/Filters/Utils/Evaluation/FilterBarEvaluation.h
new file mode 100644
--- /dev/null
+++ b/Engine/Source/Editor/Sequencer/Private/Filters/Utils/Evaluation/FilterBarEvaluation.h
@@ -0,0 +1,82 @@
+﻿// Copyright Epic Games, Inc. All Rights Reserved.
+
+#pragma once
+
+#include "FilterGroupEvaluation.h"
+#include "Containers/Array.h"
+#include "Misc/NotNull.h"
+
+class FSequencerTrackFilter_HideIsolate;
+class FSequencerTrackFilter_CustomText;
+class FSequencerTrackFilter;
+class FSequencerTrackFilter_Text;
+
+namespace UE::Sequencer
+{
+constexpr SIZE_T NumInlineFilters = 8;
+template<typename T>
+using TFilterArray = TArray<T, TInlineAllocator<NumInlineFilters>>;
+
+/**
+ * The filters run by FSequencerFilterBar. The members are specific to the needs of FSequencerFilterBar. 
+ * An item passes if the conjunction of filters passes, i.e. all filters pass on the item.
+ */
+struct FSequencerFilterContext
+{
+	/** The search box filter, if active. */
+	TOptional<TNotNull<FSequencerTrackFilter_Text*>> TextFilter;
+	
+	/** The isolate / hide / show filter, if active. */
+	TOptional<TNotNull<FSequencerTrackFilter_HideIsolate*>> HideIsolate;
+	
+	TFilterArray<TNotNull<FSequencerTrackFilter*>> CommonFilters;
+	
+	TFilterArray<TNotNull<FSequencerTrackFilter*>> InternalFilters;
+	
+	TFilterArray<TNotNull<FSequencerTrackFilter*>> CustomTextFilters;
+	
+	/** @return Whether the search can be stopped early for a subtree when all overrides are set. Text filters, etc. need to be run on all children. */
+	bool HasOnlyCachedFilters() const
+	{
+		return !TextFilter 
+			// HideIsolate doesn't implement Evaluate right now. If you implement it in the future, you can remove this here and add a FRecursiveFilterOverrides member to FSequencerFilterOverrides.
+			&& !HideIsolate 
+			&& CustomTextFilters.IsEmpty(); 
+	}
+};
+
+/** Overrides requires by FSequencerFilterBar */
+struct FSequencerFilterOverrides
+{
+	/** Overrides for FSequencerFilterContext::CommonFilters */
+	FRecursiveFilterOverrides CommonOverrides;
+	/** Overrides for FSequencerFilterContext::InternalFilters */
+	FRecursiveFilterOverrides InternalOverrides;
+	
+	bool HasOverrideForAllFilters(const FSequencerFilterContext& InContext) const
+	{
+		return CommonOverrides.Overrides.Num() == InContext.CommonFilters.Num() 
+			&& InternalOverrides.Overrides.Num() == InContext.InternalFilters.Num();
+	}
+};
+
+/** Result of calling EvaluateSequencerFilterBar. */
+struct FSequencerFilterEvaluation
+{
+	/** Indicates what to do with the item. */
+	EItemFilterState ItemResult;
+	
+	/** The new overrides to apply to children of the item. */
+	FSequencerFilterOverrides NewOverrides;
+};
+
+/** Calls FSequencerFilterBase::OnPreFilter on all filters that require it. */
+void InvokePreFilter(const FSequencerFilterContext& InFilters);
+
+/** Evaluates the filters if FSequencerFilterBar. */
+FSequencerFilterEvaluation EvaluateSequencerFilterBar(
+	const TViewModelPtr<IOutlinerExtension>& InItem, 
+	const FSequencerFilterContext& InFilters,
+	const FSequencerFilterOverrides& InOverrides
+	);
+}
diff --git a/Engine/Source/Editor/Sequencer/Private/Filters/Utils/Evaluation/FilterGroupEvaluation.cpp b/Engine/Source/Editor/Sequencer/Private/Filters/Utils/Evaluation/FilterGroupEvaluation.cpp
new file mode 100644
--- /dev/null
+++ b/Engine/Source/Editor/Sequencer/Private/Filters/Utils/Evaluation/FilterGroupEvaluation.cpp
@@ -0,0 +1,102 @@
+﻿// Copyright Epic Games, Inc. All Rights Reserved.
+
+#include "FilterGroupEvaluation.h"
+
+#include "Filters/FilterResult.h"
+#include "Filters/SequencerTrackFilterBase.h"
+#include "Filters/SequencerTrackFilterCollection.h"
+#include "MVVM/ViewModelPtr.h"
+#include "MVVM/Extensions/IOutlinerExtension.h"
+
+namespace UE::Sequencer
+{
+TOptional<EItemFilterState> FRecursiveFilterOverrides::FindOverride(const FSequencerTrackFilter& InFilter) const
+{
+	for (const FFilterOverride& Override : Overrides)
+	{
+		if (&Override.Filter == &InFilter)
+		{
+			return Override.ResultToApply;
+		}
+	}
+	
+	return {};
+}
+
+void FRecursiveFilterOverrides::AddOverride(const FSequencerTrackFilter& InFilter, EItemFilterState InOverride)
+{
+	for (FFilterOverride& Override : Overrides)
+	{
+		if (&Override.Filter == &InFilter)
+		{
+			Override.ResultToApply = InOverride;
+			return;
+		}
+	}
+	
+	Overrides.Emplace(InFilter, InOverride);
+}
+
+namespace FilterGroup
+{
+template<EItemFilterState TStopCondition, typename TCombineFunc> requires std::is_invocable_r_v<EItemFilterState, TCombineFunc, EItemFilterState, EItemFilterState>
+static FFilterGroupEvaluation Evaluate(
+	const TViewModelPtr<IOutlinerExtension>& InItem, 
+	TConstArrayView<TNotNull<FSequencerTrackFilter*>> InFilters,
+	const FRecursiveFilterOverrides& InOverrides,
+	TCombineFunc&& InCombineFunc
+	)
+{
+	FFilterGroupEvaluation Result;
+	EItemFilterState& ItemResult = Result.ItemResult;
+	FRecursiveFilterOverrides& NewOverrides = Result.NewRecursiveOverrides;
+	ItemResult = EItemFilterState::DoNotCare;
+	NewOverrides = InOverrides;
+	
+	for (const FSequencerTrackFilter* Filter : InFilters)
+	{
+		const TOptional<EItemFilterState> Override = InOverrides.FindOverride(*Filter);
+		const FFilterResult FilterResult = Override 
+			? FFilterResult(*Override, EFilterResultFlags::ApplyResultToChildren) 
+			: Filter->Evaluate(InItem);
+		
+		if (EnumHasAnyFlags(FilterResult.Flags, EFilterResultFlags::ApplyResultToChildren))
+		{
+			NewOverrides.AddOverride(*Filter, FilterResult.ItemState);
+		}
+		
+		ItemResult = InCombineFunc(Result.ItemResult, FilterResult.ItemState);
+		if (ItemResult == TStopCondition) // Short-circuit evaluation of the other filters, to skip unneeded evaluation.
+		{
+			return Result;
+		}
+	}
+	
+	return Result;
+}
+}
+
+FFilterGroupEvaluation EvaluateDisjunctive(
+	const TViewModelPtr<IOutlinerExtension>& InItem, 
+	TConstArrayView<TNotNull<FSequencerTrackFilter*>> InFilters,
+	const FRecursiveFilterOverrides& InOverrides
+	)
+{
+	return FilterGroup::Evaluate<EItemFilterState::Include>(InItem, InFilters, InOverrides, [](EItemFilterState Left, EItemFilterState Right)
+	{
+		return CombineDisjunctive(Left, Right);
+	});
+}
+
+FFilterGroupEvaluation EvaluateConjunctive(
+	const TViewModelPtr<IOutlinerExtension>& InItem,
+	TConstArrayView<TNotNull<FSequencerTrackFilter*>> InFilters, 
+	const FRecursiveFilterOverrides& InOverrides
+	)
+{
+	return FilterGroup::Evaluate<EItemFilterState::Exclude>(InItem, InFilters, InOverrides, [](EItemFilterState Left, EItemFilterState Right)
+	{
+		return CombineConjunctive(Left, Right);
+	});
+}
+}
diff --git a/Engine/Source/Editor/Sequencer/Private/Filters/Utils/Evaluation/FilterGroupEvaluation.h b/Engine/Source/Editor/Sequencer/Private/Filters/Utils/Evaluation/FilterGroupEvaluation.h
new file mode 100644
--- /dev/null
+++ b/Engine/Source/Editor/Sequencer/Private/Filters/Utils/Evaluation/FilterGroupEvaluation.h
@@ -0,0 +1,64 @@
+﻿// Copyright Epic Games, Inc. All Rights Reserved.
+
+#pragma once 
+
+#include "Filters/FilterResult.h"
+#include "MVVM/Extensions/IOutlinerExtension.h"
+#include "MVVM/ViewModelPtr.h"
+
+class FSequencerTrackFilter;
+class FSequencerTrackFilterCollection;
+
+namespace UE::Sequencer
+{
+enum class EItemFilterState : uint8;
+struct FFilterResult;
+
+/** Override of the result of a single filter */
+struct FFilterOverride
+{
+	/** The filter this override applies to */
+	const FSequencerTrackFilter& Filter;
+	/** The state to assume for any items passed to this filter */
+	EItemFilterState ResultToApply;
+
+	explicit FFilterOverride(const FSequencerTrackFilter& Filter UE_LIFETIMEBOUND, EItemFilterState ResultToApply)
+		: Filter(Filter)
+		, ResultToApply(ResultToApply)
+	{}
+};
+
+/** Holds overrides for various filters. */
+struct FRecursiveFilterOverrides
+{
+	/** Filters in this array should not be evaluated. Instead, just assume FFilterOverride::ResultToApply as result. */
+	TArray<FFilterOverride, TInlineAllocator<8>> Overrides;
+	
+	TOptional<EItemFilterState> FindOverride(const FSequencerTrackFilter& InFilter) const;
+	void AddOverride(const FSequencerTrackFilter& InFilter, EItemFilterState InOverride);
+};
+
+/** Result of EvaluateFiltersOnItem. */
+struct FFilterGroupEvaluation
+{
+	/** Indicates what to do with the item. */
+	EItemFilterState ItemResult;
+	
+	/** The new overrides to apply to children of the item. */
+	FRecursiveFilterOverrides NewRecursiveOverrides;
+};
+
+/** Runs all filters and returns Include if any filter passes. */
+FFilterGroupEvaluation EvaluateDisjunctive(
+	const TViewModelPtr<IOutlinerExtension>& InItem, 
+	TConstArrayView<TNotNull<FSequencerTrackFilter*>> InFilters,
+	const FRecursiveFilterOverrides& InOverrides
+	);
+
+/** Runs all filters and returns Include if all filters pass. */
+FFilterGroupEvaluation EvaluateConjunctive(
+	const TViewModelPtr<IOutlinerExtension>& InItem, 
+	TConstArrayView<TNotNull<FSequencerTrackFilter*>> InFilters,
+	const FRecursiveFilterOverrides& InOverrides
+	);
+}
diff --git a/Engine/Source/Editor/Sequencer/Private/Filters/Utils/FilterEvaluation.cpp b/Engine/Source/Editor/Sequencer/Private/Filters/Utils/FilterEvaluation.cpp
deleted file mode 100644
diff --git a/Engine/Source/Editor/Sequencer/Private/Filters/Utils/FilterEvaluation.h b/Engine/Source/Editor/Sequencer/Private/Filters/Utils/FilterEvaluation.h
deleted file mode 100644
diff --git a/Engine/Source/Editor/Sequencer/Public/Filters/FilterResult.h b/Engine/Source/Editor/Sequencer/Public/Filters/FilterResult.h
--- a/Engine/Source/Editor/Sequencer/Public/Filters/FilterResult.h
+++ b/Engine/Source/Editor/Sequencer/Public/Filters/FilterResult.h
@@ -15,12 +15,24 @@ enum class EItemFilterState : uint8
 	Include,
 	/** The item should not be displayed to the user. */
 	Exclude,
-	/** The filter does not care. Let other filters decide. */
+	/** The filter does not care or is not applicable. Let other filters decide. */
 	DoNotCare
 };
 
-/** @return Include if either is set to Include and the other to Include or DoNotCare, Exclude if either is set to exclude, and DoNotCare otherwise. */
-EItemFilterState Combine(EItemFilterState InLeft, EItemFilterState InRight);
+/** @return Exclude if either is set to exclude, Include if both are set to Include, and DoNotCare otherwise. */
+EItemFilterState CombineConjunctive(EItemFilterState InLeft, EItemFilterState InRight);
+/** @return Include if either is set to Include, Exclude if both are set to Exclude, and DoNotCare otherwise. */
+EItemFilterState CombineDisjunctive(EItemFilterState InLeft, EItemFilterState InRight);
+
+/** @return The result of combining any number of EItemFilterState values, folding left using the 2-param Combine. */
+template<typename... TStates>
+EItemFilterState CombineConjunctive(EItemFilterState InFirst, TStates... InRest);
+/** @return The result of combining any number of EItemFilterState values, folding left using the 2-param Combine. */
+template<typename... TStates>
+EItemFilterState CombineDisjunctive(EItemFilterState InFirst, TStates... InRest);
+
+/** @return Whether an item should be displayed */
+bool PassesFilterState(EItemFilterState InState);
 
 /** Additional info filters may want to return after having evaluated an item. */
 enum class EFilterResultFlags : uint8
@@ -47,7 +59,7 @@ struct FFilterResult
 	/** Additional information the filtered returned, e.g. that the filter need not be invoked for children. */
 	EFilterResultFlags Flags;
 
-	explicit FFilterResult(EItemFilterState ItemState, EFilterResultFlags Flags = EFilterResultFlags::None)
+	FFilterResult(EItemFilterState ItemState, EFilterResultFlags Flags = EFilterResultFlags::None)
 		: ItemState(ItemState)
 		, Flags(Flags)
 	{}
@@ -56,7 +68,7 @@ struct FFilterResult
 
 namespace UE::Sequencer
 {
-inline EItemFilterState Combine(EItemFilterState InLeft, EItemFilterState InRight)
+inline EItemFilterState CombineConjunctive(EItemFilterState InLeft, EItemFilterState InRight)
 {
 	switch (InLeft)
 	{
@@ -66,4 +78,33 @@ inline EItemFilterState Combine(EItemFilterState InLeft, EItemFilterState InRigh
 	default: checkNoEntry(); return EItemFilterState::DoNotCare;
 	}
 }
+
+inline EItemFilterState CombineDisjunctive(EItemFilterState InLeft, EItemFilterState InRight)
+{
+	switch (InLeft)
+	{
+	case EItemFilterState::Include: return EItemFilterState::Include;
+	case EItemFilterState::Exclude: return InRight == EItemFilterState::Include ? EItemFilterState::Include : InLeft;
+	case EItemFilterState::DoNotCare: return InRight;
+	default: checkNoEntry(); return EItemFilterState::DoNotCare;
+	}
+}
+
+template<typename... TStates>
+EItemFilterState CombineConjunctive(EItemFilterState InFirst, TStates... InRest)
+{
+	return CombineConjunctive(InFirst, CombineConjunctive(InRest...));
+}
+
+template<typename... TStates>
+EItemFilterState CombineDisjunctive(EItemFilterState InFirst, TStates... InRest)
+{
+	return CombineDisjunctive(InFirst, CombineDisjunctive(InRest...));
+}
+
+inline bool PassesFilterState(EItemFilterState InState)
+{
+	// We'll treat DoNotCare as Include. 
+	return InState != EItemFilterState::Exclude;
+}
 }
\ No newline at end of file
