Sequencer: Simplify filter evaluation code.
- Elimnate FFilterGroupEvaluation
- All filters now share a single FRecursiveFilterOverrides, which makes it easier to add additional filter evaluation logic to EvaluateSequencerFilterBar
- The filtering algorithm has been adjusted to compare the override num against FSequencerFilterContext::NumFilters. That is easier than the previous logic.

#Synced-CL 52283655

diff --git a/Engine/Source/Editor/Sequencer/Private/Filters/SequencerFilterBar.cpp b/Engine/Source/Editor/Sequencer/Private/Filters/SequencerFilterBar.cpp
--- a/Engine/Source/Editor/Sequencer/Private/Filters/SequencerFilterBar.cpp
+++ b/Engine/Source/Editor/Sequencer/Private/Filters/SequencerFilterBar.cpp
@@ -1622,7 +1622,7 @@ void FSequencerFilterBar::FilterNodesNew()
 bool FSequencerFilterBar::FilterNodesRecursiveNew(
 	const TViewModelPtr<IOutlinerExtension>& InStartNode,
 	const FSequencerFilterContext& InFilters,
-	const FSequencerFilterOverrides& InOverrides
+	const FRecursiveFilterOverrides& InOverrides
 	)
 {
 	TRACE_CPUPROFILER_EVENT_SCOPE(FSequencerFilterBar::FilterNodesRecursive);
@@ -1653,8 +1653,8 @@ bool FSequencerFilterBar::FilterNodesRecursiveNew(
 	const bool bAllFiltersPassed = PassesFilterState(EvaluationResult.ItemResult);
 	bool bAnyChildPassed = false;
 	
-	// Some filters, such as text filters, need to be run on all children. If we're not evaluating text filters, we can optimize by trimming the tree.
-	const bool bHasOverrideForAllFilters = InFilters.HasOnlyCachedFilters() && EvaluationResult.NewOverrides.HasOverrideForAllFilters(InFilters);
+	// Some filters, such as text filters, need to be run on all children. If every filter has overrides, we can optimize by trimming the subtree.
+	const bool bHasOverrideForAllFilters = InFilters.NumFilters() == EvaluationResult.NewOverrides.Overrides.Num();
 	if (bHasOverrideForAllFilters)
 	{
 		bAnyChildPassed = bAllFiltersPassed;
diff --git a/Engine/Source/Editor/Sequencer/Private/Filters/SequencerFilterBar.h b/Engine/Source/Editor/Sequencer/Private/Filters/SequencerFilterBar.h
--- a/Engine/Source/Editor/Sequencer/Private/Filters/SequencerFilterBar.h
+++ b/Engine/Source/Editor/Sequencer/Private/Filters/SequencerFilterBar.h
@@ -338,7 +338,7 @@ private:
 	bool FilterNodesRecursiveNew(
 		const UE::Sequencer::TViewModelPtr<UE::Sequencer::IOutlinerExtension>& InStartNode,
 		const UE::Sequencer::FSequencerFilterContext& InFilters,
-		const UE::Sequencer::FSequencerFilterOverrides& InOverrides
+		const UE::Sequencer::FRecursiveFilterOverrides& InOverrides
 		);
 	
 	/** Invoked when the text filter changes. */
diff --git a/Engine/Source/Editor/Sequencer/Private/Filters/Utils/Evaluation/FilterBarEvaluation.cpp b/Engine/Source/Editor/Sequencer/Private/Filters/Utils/Evaluation/FilterBarEvaluation.cpp
--- a/Engine/Source/Editor/Sequencer/Private/Filters/Utils/Evaluation/FilterBarEvaluation.cpp
+++ b/Engine/Source/Editor/Sequencer/Private/Filters/Utils/Evaluation/FilterBarEvaluation.cpp
@@ -10,7 +10,20 @@ namespace UE::Sequencer
 {
 void InvokePreFilter(const FSequencerFilterContext& InFilters)
 {
-	// Text filters, and isolte/hide/show do not require pre-filter.
+	if (InFilters.TextFilter)
+	{
+		InFilters.TextFilter.GetValue()->OnPreFilter();
+	}
+	
+	if (InFilters.HideIsolate)
+	{
+		InFilters.HideIsolate.GetValue()->OnPreFilter();
+	}
+	
+	for (const TNotNull<FSequencerTrackFilter*> CustomTextFilter : InFilters.CustomTextFilters)
+	{
+		CustomTextFilter->OnPreFilter();
+	}
 	
 	for (const TNotNull<FSequencerTrackFilter*> CommonFilter : InFilters.CommonFilters)
 	{
@@ -27,81 +40,60 @@ namespace FilterEval
 {
 template<typename TFilter> requires std::is_base_of_v<FSequencerTrackFilter, TFilter>
 static EItemFilterState EvaluateFilter(
-	const TViewModelPtr<IOutlinerExtension>& InItem, const TOptional<TNotNull<TFilter*>> InFilter
-	)
-{
-	if (!InFilter)
-	{
-		return EItemFilterState::DoNotCare;
-	}
-	
-	return InFilter.GetValue()->PassesFilter(InItem) ? EItemFilterState::Include : EItemFilterState::Exclude;
-}
-
-static EItemFilterState EvaluateFilters(
-	const TViewModelPtr<IOutlinerExtension>& InItem, const TConstArrayView<TNotNull<FSequencerTrackFilter*>> InFilters
+	const TViewModelPtr<IOutlinerExtension>& InItem, 
+	const TOptional<TNotNull<TFilter*>> InFilter,
+	FRecursiveFilterOverrides& InOverrides
 	)
 {
-	if (InFilters.Num() == 0)
-	{
-		return EItemFilterState::DoNotCare;
-	}
-	
-	const bool bPassedAll = Algo::AllOf(InFilters, [&InItem](TNotNull<FSequencerTrackFilter*> CustomTextFilter)
-	{
-		return CustomTextFilter->PassesFilter(InItem);
-	});
-	return bPassedAll? EItemFilterState::Include : EItemFilterState::Exclude;
+	return InFilter ? EvaluateSingle(InItem, *InFilter.GetValue(), InOverrides) : EItemFilterState::DoNotCare;
 }
 }
 
 FSequencerFilterEvaluation EvaluateSequencerFilterBar(
 	const TViewModelPtr<IOutlinerExtension>& InItem, 
 	const FSequencerFilterContext& InFilters,
-	const FSequencerFilterOverrides& InOverrides
+	const FRecursiveFilterOverrides& InOverrides
 	)
 {
-	FSequencerFilterOverrides NewOverrides = InOverrides;
+	FRecursiveFilterOverrides NewOverrides = InOverrides;
 	
 	// The below is sorted in the order we expect to evaluate the fastest.
 	
-	const EItemFilterState HideIsolateFilterResult = FilterEval::EvaluateFilter(InItem, InFilters.HideIsolate);
+	const EItemFilterState HideIsolateFilterResult = FilterEval::EvaluateFilter(InItem, InFilters.HideIsolate, NewOverrides);
 	if (!PassesFilterState(HideIsolateFilterResult))
 	{
 		return FSequencerFilterEvaluation(HideIsolateFilterResult, NewOverrides);
 	}
 	
-	const EItemFilterState TextFilterResult = FilterEval::EvaluateFilter(InItem, InFilters.TextFilter);
+	const EItemFilterState TextFilterResult = FilterEval::EvaluateFilter(InItem, InFilters.TextFilter, NewOverrides);
 	if (!PassesFilterState(TextFilterResult))
 	{
 		return FSequencerFilterEvaluation(TextFilterResult, NewOverrides);
 	}
 	
-	const EItemFilterState CustomTextResult = FilterEval::EvaluateFilters(InItem, InFilters.CustomTextFilters);
+	const EItemFilterState CustomTextResult = EvaluateConjunctive(InItem, InFilters.CustomTextFilters, NewOverrides);
 	if (!PassesFilterState(CustomTextResult))
 	{
 		return FSequencerFilterEvaluation(CustomTextResult, NewOverrides);
 	}
 	
 	// The internal filters act as a pre-filter... before user filters are applied.
-	const FFilterGroupEvaluation InternalEval = EvaluateConjunctive(InItem, InFilters.InternalFilters, InOverrides.InternalOverrides);
-	NewOverrides.InternalOverrides = InternalEval.NewRecursiveOverrides;
-	if (!PassesFilterState(InternalEval.ItemResult))
+	const EItemFilterState InternalEval = EvaluateConjunctive(InItem, InFilters.InternalFilters, NewOverrides);
+	if (!PassesFilterState(InternalEval))
 	{
-		return FSequencerFilterEvaluation(InternalEval.ItemResult, NewOverrides);
+		return FSequencerFilterEvaluation(InternalEval, NewOverrides);
 	}
 	
 	// The filter pills that the user manually checks and unchecks add or remove items... which encodes an OR condition
-	const FFilterGroupEvaluation CommonEval = EvaluateDisjunctive(InItem, InFilters.CommonFilters, InOverrides.CommonOverrides);
-	NewOverrides.CommonOverrides = CommonEval.NewRecursiveOverrides;
-	if (!PassesFilterState(CommonEval.ItemResult))
+	const EItemFilterState CommonEval = EvaluateDisjunctive(InItem, InFilters.CommonFilters, NewOverrides);
+	if (!PassesFilterState(CommonEval))
 	{
-		return FSequencerFilterEvaluation(CommonEval.ItemResult, NewOverrides);
+		return FSequencerFilterEvaluation(CommonEval, NewOverrides);
 	}
 	
-	// All filter results must pass
+	// All filter results must pass. Effectively, this will combine to Include or DoNotCare.
 	const EItemFilterState CombinedResult = CombineConjunctive(
-		TextFilterResult, HideIsolateFilterResult, CommonEval.ItemResult, InternalEval.ItemResult, CustomTextResult
+		TextFilterResult, HideIsolateFilterResult, CommonEval, InternalEval, CustomTextResult
 		);
 	return FSequencerFilterEvaluation(CombinedResult, NewOverrides);
 }
diff --git a/Engine/Source/Editor/Sequencer/Private/Filters/Utils/Evaluation/FilterBarEvaluation.h b/Engine/Source/Editor/Sequencer/Private/Filters/Utils/Evaluation/FilterBarEvaluation.h
--- a/Engine/Source/Editor/Sequencer/Private/Filters/Utils/Evaluation/FilterBarEvaluation.h
+++ b/Engine/Source/Editor/Sequencer/Private/Filters/Utils/Evaluation/FilterBarEvaluation.h
@@ -35,28 +35,14 @@ struct FSequencerFilterContext
 	
 	TFilterArray<TNotNull<FSequencerTrackFilter*>> CustomTextFilters;
 	
-	/** @return Whether the search can be stopped early for a subtree when all overrides are set. Text filters, etc. need to be run on all children. */
-	bool HasOnlyCachedFilters() const
+	/** @return The number of active filters. */
+	int32 NumFilters() const
 	{
-		return !TextFilter 
-			// HideIsolate doesn't implement Evaluate right now. If you implement it in the future, you can remove this here and add a FRecursiveFilterOverrides member to FSequencerFilterOverrides.
-			&& !HideIsolate 
-			&& CustomTextFilters.IsEmpty(); 
-	}
-};
-
-/** Overrides requires by FSequencerFilterBar */
-struct FSequencerFilterOverrides
-{
-	/** Overrides for FSequencerFilterContext::CommonFilters */
-	FRecursiveFilterOverrides CommonOverrides;
-	/** Overrides for FSequencerFilterContext::InternalFilters */
-	FRecursiveFilterOverrides InternalOverrides;
-	
-	bool HasOverrideForAllFilters(const FSequencerFilterContext& InContext) const
-	{
-		return CommonOverrides.Overrides.Num() == InContext.CommonFilters.Num() 
-			&& InternalOverrides.Overrides.Num() == InContext.InternalFilters.Num();
+		return static_cast<int32>(TextFilter.IsSet()) 
+			+ static_cast<int32>(HideIsolate.IsSet()) 
+			+ CommonFilters.Num()
+			+ InternalFilters.Num()
+			+ CustomTextFilters.Num();
 	}
 };
 
@@ -67,7 +53,7 @@ struct FSequencerFilterEvaluation
 	EItemFilterState ItemResult;
 	
 	/** The new overrides to apply to children of the item. */
-	FSequencerFilterOverrides NewOverrides;
+	FRecursiveFilterOverrides NewOverrides;
 };
 
 /** Calls FSequencerFilterBase::OnPreFilter on all filters that require it. */
@@ -77,6 +63,6 @@ void InvokePreFilter(const FSequencerFilterContext& InFilters);
 FSequencerFilterEvaluation EvaluateSequencerFilterBar(
 	const TViewModelPtr<IOutlinerExtension>& InItem, 
 	const FSequencerFilterContext& InFilters,
-	const FSequencerFilterOverrides& InOverrides
+	const FRecursiveFilterOverrides& InOverrides
 	);
 }
diff --git a/Engine/Source/Editor/Sequencer/Private/Filters/Utils/Evaluation/FilterGroupEvaluation.cpp b/Engine/Source/Editor/Sequencer/Private/Filters/Utils/Evaluation/FilterGroupEvaluation.cpp
--- a/Engine/Source/Editor/Sequencer/Private/Filters/Utils/Evaluation/FilterGroupEvaluation.cpp
+++ b/Engine/Source/Editor/Sequencer/Private/Filters/Utils/Evaluation/FilterGroupEvaluation.cpp
@@ -40,63 +40,69 @@ void FRecursiveFilterOverrides::AddOverride(const FSequencerTrackFilter& InFilte
 namespace FilterGroup
 {
 template<EItemFilterState TStopCondition, typename TCombineFunc> requires std::is_invocable_r_v<EItemFilterState, TCombineFunc, EItemFilterState, EItemFilterState>
-static FFilterGroupEvaluation Evaluate(
+static EItemFilterState Evaluate(
 	const TViewModelPtr<IOutlinerExtension>& InItem, 
 	TConstArrayView<TNotNull<FSequencerTrackFilter*>> InFilters,
-	const FRecursiveFilterOverrides& InOverrides,
+	FRecursiveFilterOverrides& InOutOverrides,
 	TCombineFunc&& InCombineFunc
 	)
 {
-	FFilterGroupEvaluation Result;
-	EItemFilterState& ItemResult = Result.ItemResult;
-	FRecursiveFilterOverrides& NewOverrides = Result.NewRecursiveOverrides;
-	ItemResult = EItemFilterState::DoNotCare;
-	NewOverrides = InOverrides;
+	EItemFilterState CombinedResult = EItemFilterState::DoNotCare;
 	
 	for (const FSequencerTrackFilter* Filter : InFilters)
 	{
-		const TOptional<EItemFilterState> Override = InOverrides.FindOverride(*Filter);
-		const FFilterResult FilterResult = Override 
-			? FFilterResult(*Override, EFilterResultFlags::ApplyResultToChildren) 
-			: Filter->Evaluate(InItem);
-		
-		if (EnumHasAnyFlags(FilterResult.Flags, EFilterResultFlags::ApplyResultToChildren))
-		{
-			NewOverrides.AddOverride(*Filter, FilterResult.ItemState);
-		}
-		
-		ItemResult = InCombineFunc(Result.ItemResult, FilterResult.ItemState);
-		if (ItemResult == TStopCondition) // Short-circuit evaluation of the other filters, to skip unneeded evaluation.
+		const EItemFilterState FilterResult = EvaluateSingle(InItem, *Filter, InOutOverrides);
+		CombinedResult = InCombineFunc(CombinedResult, FilterResult);
+		if (CombinedResult == TStopCondition) // Short-circuit evaluation of the other filters, to skip unneeded evaluation.
 		{
-			return Result;
+			return CombinedResult;
 		}
 	}
 	
-	return Result;
+	return CombinedResult;
 }
 }
 
-FFilterGroupEvaluation EvaluateDisjunctive(
+EItemFilterState EvaluateDisjunctive(
 	const TViewModelPtr<IOutlinerExtension>& InItem, 
 	TConstArrayView<TNotNull<FSequencerTrackFilter*>> InFilters,
-	const FRecursiveFilterOverrides& InOverrides
+	FRecursiveFilterOverrides& InOutOverrides
 	)
 {
-	return FilterGroup::Evaluate<EItemFilterState::Include>(InItem, InFilters, InOverrides, [](EItemFilterState Left, EItemFilterState Right)
+	return FilterGroup::Evaluate<EItemFilterState::Include>(InItem, InFilters, InOutOverrides, [](EItemFilterState Left, EItemFilterState Right)
 	{
 		return CombineDisjunctive(Left, Right);
 	});
 }
 
-FFilterGroupEvaluation EvaluateConjunctive(
+EItemFilterState EvaluateConjunctive(
 	const TViewModelPtr<IOutlinerExtension>& InItem,
 	TConstArrayView<TNotNull<FSequencerTrackFilter*>> InFilters, 
-	const FRecursiveFilterOverrides& InOverrides
+	FRecursiveFilterOverrides& InOutOverrides
 	)
 {
-	return FilterGroup::Evaluate<EItemFilterState::Exclude>(InItem, InFilters, InOverrides, [](EItemFilterState Left, EItemFilterState Right)
+	return FilterGroup::Evaluate<EItemFilterState::Exclude>(InItem, InFilters, InOutOverrides, [](EItemFilterState Left, EItemFilterState Right)
 	{
 		return CombineConjunctive(Left, Right);
 	});
 }
+
+EItemFilterState EvaluateSingle(
+	const TViewModelPtr<IOutlinerExtension>& InItem,
+	const FSequencerTrackFilter& InFilter,
+	FRecursiveFilterOverrides& InOutOverrides
+	)
+{
+	const TOptional<EItemFilterState> Override = InOutOverrides.FindOverride(InFilter);
+	const FFilterResult FilterResult = Override 
+		? FFilterResult(*Override, EFilterResultFlags::ApplyResultToChildren) 
+		: InFilter.Evaluate(InItem);
+		
+	if (EnumHasAnyFlags(FilterResult.Flags, EFilterResultFlags::ApplyResultToChildren))
+	{
+		InOutOverrides.AddOverride(InFilter, FilterResult.ItemState);
+	}
+	
+	return FilterResult.ItemState;
+}
 }
diff --git a/Engine/Source/Editor/Sequencer/Private/Filters/Utils/Evaluation/FilterGroupEvaluation.h b/Engine/Source/Editor/Sequencer/Private/Filters/Utils/Evaluation/FilterGroupEvaluation.h
--- a/Engine/Source/Editor/Sequencer/Private/Filters/Utils/Evaluation/FilterGroupEvaluation.h
+++ b/Engine/Source/Editor/Sequencer/Private/Filters/Utils/Evaluation/FilterGroupEvaluation.h
@@ -38,27 +38,24 @@ struct FRecursiveFilterOverrides
 	void AddOverride(const FSequencerTrackFilter& InFilter, EItemFilterState InOverride);
 };
 
-/** Result of EvaluateFiltersOnItem. */
-struct FFilterGroupEvaluation
-{
-	/** Indicates what to do with the item. */
-	EItemFilterState ItemResult;
-	
-	/** The new overrides to apply to children of the item. */
-	FRecursiveFilterOverrides NewRecursiveOverrides;
-};
-
 /** Runs all filters and returns Include if any filter passes. */
-FFilterGroupEvaluation EvaluateDisjunctive(
+EItemFilterState EvaluateDisjunctive(
 	const TViewModelPtr<IOutlinerExtension>& InItem, 
 	TConstArrayView<TNotNull<FSequencerTrackFilter*>> InFilters,
-	const FRecursiveFilterOverrides& InOverrides
+	FRecursiveFilterOverrides& InOutOverrides
 	);
 
 /** Runs all filters and returns Include if all filters pass. */
-FFilterGroupEvaluation EvaluateConjunctive(
+EItemFilterState EvaluateConjunctive(
 	const TViewModelPtr<IOutlinerExtension>& InItem, 
 	TConstArrayView<TNotNull<FSequencerTrackFilter*>> InFilters,
-	const FRecursiveFilterOverrides& InOverrides
+	FRecursiveFilterOverrides& InOutOverrides
+	);
+
+/** Runs the filter and appends its result to FRecursiveFilterOverrides if the filter reports that the result applies to child items. */
+EItemFilterState EvaluateSingle(
+	const TViewModelPtr<IOutlinerExtension>& InItem,
+	const FSequencerTrackFilter& InFilter,
+	FRecursiveFilterOverrides& InOutOverrides
 	);
 }
