Sequencer: Make filter bar visibile by default

#Synced-CL 52283655

diff --git a/Engine/Source/Editor/Sequencer/Public/Filters/SequencerFilterAreaConfig.h b/Engine/Source/Editor/Sequencer/Public/Filters/SequencerFilterAreaConfig.h
--- a/Engine/Source/Editor/Sequencer/Public/Filters/SequencerFilterAreaConfig.h
+++ b/Engine/Source/Editor/Sequencer/Public/Filters/SequencerFilterAreaConfig.h
@@ -25,7 +25,7 @@ private:
 	
 	/** Whether the filter bar should be visible. */
 	UPROPERTY()
-	bool bIsFilterBarVisible = false;
+	bool bIsFilterBarVisible = true;
 	
 	/** Whether when you unlink your filters, you want the filters in your linked filter bar to be copied into the unlinked filter bar. */
 	UPROPERTY()
