Sequencer: Fix that using "Hide Selected Tracks" filter in Sequencer leads that all tracks aren't displayed in the Curve Editor. This only happens when "Isolate Curve Editor to Selection" is enabled.

#Synced-CL 52378315

diff --git a/Engine/Source/Editor/Sequencer/Private/Filters/Filters/SequencerTrackFilter_HideIsolate.cpp b/Engine/Source/Editor/Sequencer/Private/Filters/Filters/SequencerTrackFilter_HideIsolate.cpp
--- a/Engine/Source/Editor/Sequencer/Private/Filters/Filters/SequencerTrackFilter_HideIsolate.cpp
+++ b/Engine/Source/Editor/Sequencer/Private/Filters/Filters/SequencerTrackFilter_HideIsolate.cpp
@@ -2,7 +2,9 @@
 
 #include "Filters/Filters/SequencerTrackFilter_HideIsolate.h"
 #include "Filters/SequencerTrackFilterCommands.h"
+#include "MVVM/Selection/Selection.h"
 #include "MVVM/ViewModels/CategoryModel.h"
+#include "MVVM/ViewModels/SequencerEditorViewModel.h"
 #include "MVVM/ViewModels/TrackModel.h"
 #include "MVVM/ViewModels/ViewModelIterators.h"
 
@@ -85,19 +87,32 @@ void FSequencerTrackFilter_HideIsolate::HideTracks(const TSet<TWeakViewModelPtr<
 		HiddenTracks.Empty();
 	}
 
-	for (const TWeakViewModelPtr<IOutlinerExtension>& TrackModekWeak : InTracks)
+	const TSharedPtr<FSequencerSelection>& Selection = GetSequencer().GetViewModel()->GetSelection();
 	{
-		if (const TViewModelPtr<IOutlinerExtension> TrackModel = TrackModekWeak.Pin())
+		const FSelectionEventSuppressor SelectionChange = Selection->SuppressEvents();
+		const auto MarkHidden = [this, &Selection](const TViewModelPtr<IOutlinerExtension>& InViewModel)
 		{
-			HiddenTracks.Add(TrackModel);
-
-			for (const TViewModelPtr<IOutlinerExtension>& ChildNode : TrackModel.AsModel()->GetDescendantsOfType<IOutlinerExtension>())
+			HiddenTracks.Add(InViewModel);
+		
+			// Conceptually, it does not make sense for an item to stay selected when the user has actively hidden it.
+			// If we didn't remove it from selection, USequencerSettings::bIsolateCurveEditorToSelection would cause the Curve Editor to be blank.
+			Selection->Outliner.Deselect(InViewModel);
+		};
+	
+		for (const TWeakViewModelPtr<IOutlinerExtension>& TrackModekWeak : InTracks)
+		{
+			if (const TViewModelPtr<IOutlinerExtension> TrackModel = TrackModekWeak.Pin())
 			{
-				HiddenTracks.Add(ChildNode);
+				MarkHidden(TrackModel);
+
+				for (const TViewModelPtr<IOutlinerExtension>& ChildNode : TrackModel.AsModel()->GetDescendantsOfType<IOutlinerExtension>())
+				{
+					MarkHidden(ChildNode);
+				}
 			}
 		}
 	}
-
+	
 	BroadcastChangedEvent();
 }
 
