Sequencer: Fix play not restarting when playhead is at the end

Follow up to 51793960

#Synced-CL 52445895

diff --git a/Engine/Source/Editor/Sequencer/Private/Sequencer.cpp b/Engine/Source/Editor/Sequencer/Private/Sequencer.cpp
--- a/Engine/Source/Editor/Sequencer/Private/Sequencer.cpp
+++ b/Engine/Source/Editor/Sequencer/Private/Sequencer.cpp
@@ -5282,16 +5282,30 @@ FReply FSequencer::OnPlay(bool bTogglePlay)
 
 		if (PlaybackSpeed > 0)
 		{
-			if (MinInclusiveTime.IsSet() && GlobalTime <= MinInclusiveTime.GetValue())
+			if (MinInclusiveTime.IsSet())
 			{
-				SetGlobalTime(MinInclusiveTime.GetValue());
+				if (GlobalTime <= MinInclusiveTime.GetValue())
+				{
+					SetGlobalTime(MinInclusiveTime.GetValue());
+				}
+				else if (MaxInclusiveTime.IsSet() && GlobalTime >= MaxInclusiveTime.GetValue())
+				{
+					SetGlobalTime(MinInclusiveTime.GetValue());
+				}
 			}
 		}
 		else
 		{
-			if (MaxInclusiveTime.IsSet() && GlobalTime >= MaxInclusiveTime.GetValue())
+			if (MaxInclusiveTime.IsSet())
 			{
-				SetGlobalTime(MaxInclusiveTime.GetValue());
+				if (GlobalTime >= MaxInclusiveTime.GetValue())
+				{
+					SetGlobalTime(MaxInclusiveTime.GetValue());
+				}
+				else if (MinInclusiveTime.IsSet() && GlobalTime <= MinInclusiveTime.GetValue())
+				{
+					SetGlobalTime(MaxInclusiveTime.GetValue());
+				}
 			}
 		}
 
