Sequencer: Fix selection range end not being set as exclusive when dragging with the time slider.

Ths same fix was made a long while ago for the playback range end.

#Synced-CL 52445895

diff --git a/Engine/Source/Editor/Sequencer/Private/SequencerTimeSliderController.cpp b/Engine/Source/Editor/Sequencer/Private/SequencerTimeSliderController.cpp
--- a/Engine/Source/Editor/Sequencer/Private/SequencerTimeSliderController.cpp
+++ b/Engine/Source/Editor/Sequencer/Private/SequencerTimeSliderController.cpp
@@ -2659,7 +2659,7 @@ void FSequencerTimeSliderController::SetSelectionRangeEnd(FFrameNumber NewEnd)
 	}
 	else if (NewEnd >= UE::MovieScene::DiscreteInclusiveLower(SelectionRange))
 	{
-		TimeSliderArgs.OnSelectionRangeChanged.ExecuteIfBound(TRange<FFrameNumber>(SelectionRange.GetLowerBound(), NewEnd));
+		TimeSliderArgs.OnSelectionRangeChanged.ExecuteIfBound(TRange<FFrameNumber>(SelectionRange.GetLowerBound(), TRangeBound<FFrameNumber>::Exclusive(NewEnd)));
 	}
 }
 
