Sequencer: Adds the FSequencerTrackFilter_StickySelection filter.

Prevents clicks in the Outliner or Curve Editor from changing which items the "Selected" filter shows.
Even though clicking items in the Outliner or Curve Editor changes the selection internally, the items clicked will continue to be filtered in until the user changes the selection externally, e.g. by clicking in the viewport.

Effective changes:
- New FSequencerTrackFilter_HierarchyBased base class. Intended for filters that implement Evaluate, which is the version that allows for trimming of the search tree.
- HierarchyFilterCache: Extracted selection filtering logic from Selection filter to separate functions. This allows the logic to be reused in the StickySelection filter
- New FSequencer::IsUpdatingSequencerSelection exposed
- FSequencerTrackFilter_StickySelection: This filter is only active when FSequencerTrackFilter_Selected is active. This logic is managed in FSequencerFilterBar. Currently, only the new algorithm adds this filter... but that is fine since we plan to get rid of the old algorithm before 5.8 ships.

#Synced-CL 52378315

diff --git a/Engine/Plugins/Animation/ControlRig/Source/ControlRigEditor/Private/Sequencer/ControlRigParameterTrackEditor.h b/Engine/Plugins/Animation/ControlRig/Source/ControlRigEditor/Private/Sequencer/ControlRigParameterTrackEditor.h
--- a/Engine/Plugins/Animation/ControlRig/Source/ControlRigEditor/Private/Sequencer/ControlRigParameterTrackEditor.h
+++ b/Engine/Plugins/Animation/ControlRig/Source/ControlRigEditor/Private/Sequencer/ControlRigParameterTrackEditor.h
@@ -158,6 +158,10 @@ public:
 
 	//function to set keys on controls
 	void SetKeysOnControls(const TMap<UControlRig*, TArray<FRigElementKey>>& Controls, EMovieSceneTransformChannel Channel);
+	
+	/** Whether the selection is currently being changed by the track editor. */
+	bool IsDoingSelection() const { return bIsDoingSelection; }
+	
 private:
 
 	void ProcessKeyOperation(UObject* ObjectToKey, UMovieSceneControlRigParameterTrack* Track, TArrayView<const UE::Sequencer::FKeySectionOperation> SectionsToKey, ISequencer& InSequencer, FFrameNumber KeyTime, TArray<UE::Sequencer::FAddKeyResult>* OutResults);
diff --git a/Engine/Plugins/Animation/ControlRig/Source/ControlRigEditor/Private/Sequencer/ControlRigSelectionFilterExtension.cpp b/Engine/Plugins/Animation/ControlRig/Source/ControlRigEditor/Private/Sequencer/ControlRigSelectionFilterExtension.cpp
deleted file mode 100644
diff --git a/Engine/Plugins/Animation/ControlRig/Source/ControlRigEditor/Private/Sequencer/ControlRigSelectionFilterExtension.h b/Engine/Plugins/Animation/ControlRig/Source/ControlRigEditor/Private/Sequencer/ControlRigSelectionFilterExtension.h
deleted file mode 100644
diff --git a/Engine/Plugins/Animation/ControlRig/Source/ControlRigEditor/Private/Sequencer/Extensions/ControlRigSelectionFilterExtension.cpp b/Engine/Plugins/Animation/ControlRig/Source/ControlRigEditor/Private/Sequencer/Extensions/ControlRigSelectionFilterExtension.cpp
new file mode 100644
--- /dev/null
+++ b/Engine/Plugins/Animation/ControlRig/Source/ControlRigEditor/Private/Sequencer/Extensions/ControlRigSelectionFilterExtension.cpp
@@ -0,0 +1,105 @@
+﻿// Copyright Epic Games, Inc. All Rights Reserved.
+
+#include "ControlRigSelectionFilterExtension.h"
+
+#include "Filters/Filters/SequencerTrackFilter_Selected.h"
+#include "MVVM/Selection/Selection.h"
+#include "MVVM/ViewModels/CategoryModel.h"
+#include "MVVM/ViewModels/SequencerEditorViewModel.h"
+#include "MVVM/ViewModels/TrackModel.h"
+#include "MVVM/ViewModels/ViewModelIterators.h"
+#include "Sequencer/MovieSceneControlRigParameterTrack.h"
+
+namespace UE::ControlRig::SelectionFilterDetail
+{
+using EBreakBehavior = UControlRigSelectionFilterExtension::EBreakBehavior;
+
+/** 
+ * If InSelected is under any FViewModel representing a control, invoke InCallback with all of children of that control.
+ * 
+ * For example, if "Dragon.ChestCtrl.Location.X" is selected, then invoke with "Dragon.ChestCtrl.Location.Y", "Dragon.ChestCtrl.Rotation.X", etc.
+ * Invokes nothing if "Dragon.ChestCtrl" or above is selected.
+ */
+static void EnumerateSiblingsIfUnderControlRigTrack(
+	const Sequencer::FViewModelPtr& InSelected, 
+	TFunctionRef<EBreakBehavior(const Sequencer::FViewModelPtr&)> InCallback,
+	TArray<Sequencer::FViewModelPtr>& InAlreadyVisited
+	)
+{
+	using namespace UE::Sequencer;
+	// Skip if "Dragon.ChestCtrl" was selected. If "Dragon.Chest.Control.Foo" was selected, proceed.
+	const bool bIsUnderControl = CastViewModel<FChannelGroupModel>(InSelected) != nullptr
+		|| CastViewModel<FChannelModel>(InSelected) != nullptr;
+	if (!bIsUnderControl)
+	{
+		return;
+	}
+	
+	for (const FViewModelPtr& Current : FParentModelIterator(InSelected, true))
+	{
+		// This is an optimization to avoid reiterating the subhierarchy in case the user has selected e.g. "Location.X" and "Location.Y".
+		if (InAlreadyVisited.Contains(Current))
+		{
+			return;
+		}
+		InAlreadyVisited.Add(Current);
+			
+		const FViewModelPtr Parent = Current->GetParent();
+		const TViewModelPtr<FTrackModel> TrackModel = CastViewModel<FTrackModel>(Parent);
+		if (!TrackModel)
+		{
+			continue;
+		}
+			
+		UMovieSceneControlRigParameterTrack* ControlRigTrack = Cast<UMovieSceneControlRigParameterTrack>(TrackModel->GetTrack());
+		if (!ControlRigTrack)
+		{
+			continue;
+		}
+			
+		for (const FViewModelPtr& Child : FParentFirstChildIterator(Current))
+		{
+			if (InCallback(Child) == EBreakBehavior::Break)
+			{
+				return;
+			}
+		}
+	}
+}
+
+static void EnumerateParentControlChildren(
+	const ISequencer& InSequencer,
+	TFunctionRef<EBreakBehavior(const Sequencer::FViewModelPtr&)> InCallback
+	)
+{
+	using namespace UE::Sequencer;
+	const TSharedPtr<FSequencerEditorViewModel> ViewModel = InSequencer.GetViewModel();
+	if (!ViewModel.IsValid())
+	{
+		return;
+	}
+
+	const TSharedPtr<FSequencerSelection> SequencerSelection = ViewModel->GetSelection();
+	if (!SequencerSelection.IsValid())
+	{
+		return;
+	}
+	
+	TArray<FViewModelPtr> AllReadyVisited;
+	for (const FViewModelPtr& Selected : SequencerSelection->Outliner)
+	{
+		EnumerateSiblingsIfUnderControlRigTrack(Selected, InCallback, AllReadyVisited);
+	}
+}
+}
+
+void UControlRigSelectionFilterExtension::EnumerateViewModelsConsideredAsSelected(
+	const FSequencerTrackFilter_Selected& InFilter,
+	TFunctionRef<EBreakBehavior(const UE::Sequencer::FViewModelPtr&)> InCallback
+	)
+{
+	using namespace UE::ControlRig::SelectionFilterDetail;
+	
+	// E.g. if "Dragon.ChestCtrl.Location.X" is selected, then also show the other child properties: Location, Rotation, etc.
+	EnumerateParentControlChildren(InFilter.GetSequencer(), InCallback);
+}
diff --git a/Engine/Plugins/Animation/ControlRig/Source/ControlRigEditor/Private/Sequencer/Extensions/ControlRigSelectionFilterExtension.h b/Engine/Plugins/Animation/ControlRig/Source/ControlRigEditor/Private/Sequencer/Extensions/ControlRigSelectionFilterExtension.h
new file mode 100644
--- /dev/null
+++ b/Engine/Plugins/Animation/ControlRig/Source/ControlRigEditor/Private/Sequencer/Extensions/ControlRigSelectionFilterExtension.h
@@ -0,0 +1,21 @@
+﻿// Copyright Epic Games, Inc. All Rights Reserved.
+
+#pragma once
+
+#include "Filters/SequencerTrackFilterSelectedExtension.h"
+#include "ControlRigSelectionFilterExtension.generated.h"
+
+/**
+ * Extends FSequencerTrackFilter_Selected such that selecting e.g. Dragon.ChestCtrl.Location.X shows Location.Y, Rotation, Scale, etc., as well.
+ */
+UCLASS()
+class UControlRigSelectionFilterExtension : public USequencerTrackFilterSelectedExtension
+{
+	GENERATED_BODY()
+public:
+	
+	virtual void EnumerateViewModelsConsideredAsSelected(
+		const FSequencerTrackFilter_Selected& InFilter,
+		TFunctionRef<EBreakBehavior(const UE::Sequencer::FViewModelPtr&)> InCallback
+		) override;
+};
diff --git a/Engine/Plugins/Animation/ControlRig/Source/ControlRigEditor/Private/Sequencer/Extensions/ControlRigStickySelectionFilterExtension.cpp b/Engine/Plugins/Animation/ControlRig/Source/ControlRigEditor/Private/Sequencer/Extensions/ControlRigStickySelectionFilterExtension.cpp
new file mode 100644
--- /dev/null
+++ b/Engine/Plugins/Animation/ControlRig/Source/ControlRigEditor/Private/Sequencer/Extensions/ControlRigStickySelectionFilterExtension.cpp
@@ -0,0 +1,11 @@
+﻿// Copyright Epic Games, Inc. All Rights Reserved.
+
+#include "ControlRigStickySelectionFilterExtension.h"
+
+#include "Sequencer/ControlRigParameterTrackEditor.h"
+
+bool UControlRigStickySelectionFilterExtension::IsPerformingSelectionThatShouldClearStickySelection(const ISequencer& InSequencer)
+{
+	const FControlRigParameterTrackEditor* Editor = FControlRigParameterTrackEditor::CurrentEditor;
+	return Editor && Editor->IsDoingSelection();
+}
diff --git a/Engine/Plugins/Animation/ControlRig/Source/ControlRigEditor/Private/Sequencer/Extensions/ControlRigStickySelectionFilterExtension.h b/Engine/Plugins/Animation/ControlRig/Source/ControlRigEditor/Private/Sequencer/Extensions/ControlRigStickySelectionFilterExtension.h
new file mode 100644
--- /dev/null
+++ b/Engine/Plugins/Animation/ControlRig/Source/ControlRigEditor/Private/Sequencer/Extensions/ControlRigStickySelectionFilterExtension.h
@@ -0,0 +1,16 @@
+﻿// Copyright Epic Games, Inc. All Rights Reserved.
+
+#pragma once
+
+#include "Filters/SequencerTrackFilterStickySelectionExtension.h"
+#include "ControlRigStickySelectionFilterExtension.generated.h"
+
+/** Extends FSequencerTrackFilter_StickySelection such that clicking control rig controls in the viewport clears sticky selection. */
+UCLASS()
+class UControlRigStickySelectionFilterExtension : public USequencerTrackFilterStickySelectionExtension
+{
+	GENERATED_BODY()
+public:
+	
+	virtual bool IsPerformingSelectionThatShouldClearStickySelection(const ISequencer& InSequencer) override;
+};
diff --git a/Engine/Source/Editor/Sequencer/Private/Filters/Filters/SequencerTrackFilter_Selected.cpp b/Engine/Source/Editor/Sequencer/Private/Filters/Filters/SequencerTrackFilter_Selected.cpp
--- a/Engine/Source/Editor/Sequencer/Private/Filters/Filters/SequencerTrackFilter_Selected.cpp
+++ b/Engine/Source/Editor/Sequencer/Private/Filters/Filters/SequencerTrackFilter_Selected.cpp
@@ -10,6 +10,7 @@
 #include "MVVM/Selection/Selection.h"
 #include "MVVM/ViewModels/SequencerEditorViewModel.h"
 #include "Selection.h"
+#include "Filters/Utils/HierarchyFilterCache.h"
 #include "UObject/UObjectIterator.h"
 
 using namespace UE::Sequencer;
@@ -80,7 +81,7 @@ static bool IsPartOfExtendedSelection(
 }
 
 FSequencerTrackFilter_Selected::FSequencerTrackFilter_Selected(ISequencerTrackFilters& InFilterInterface, TSharedPtr<FFilterCategory> InCategory)
-	: FSequencerTrackFilter(InFilterInterface, MoveTemp(InCategory))
+	: FSequencerTrackFilter_HierarchyBased(InFilterInterface, MoveTemp(InCategory))
 	, CachedWeakExtensions(SelectionDetail::GetSelectionExtensions())
 {}
 
@@ -135,16 +136,6 @@ void FSequencerTrackFilter_Selected::ActiveStateChanged(const bool bInActive)
 	}
 }
 
-FFilterResult FSequencerTrackFilter_Selected::Evaluate(FSequencerTrackFilterType InItem) const
-{
-	const EFilterResultFlags* Flags = IncludedViewModels.Find(InItem);
-	const bool bPassesFilter = Flags != nullptr;
-	
-	const EItemFilterState ResultState = bPassesFilter ? EItemFilterState::Include : EItemFilterState::Exclude;
-	const EFilterResultFlags ResultFlags = Flags ? *Flags : EFilterResultFlags::None;
-	return FFilterResult(ResultState, ResultFlags);
-}
-
 FText FSequencerTrackFilter_Selected::GetDisplayName() const
 {
 	return LOCTEXT("SequencerTrackFilter_Selected", "Selected");
@@ -164,8 +155,8 @@ void FSequencerTrackFilter_Selected::OnPreFilter()
 	
 	TRACE_CPUPROFILER_EVENT_SCOPE(FSequencerTrackFilter_Selected::OnPreFilter);
 	bNeedsPrefilterRebuild = false;
-	IncludedViewModels.Empty(IncludedViewModels.Num());
-	ON_SCOPE_EXIT{ IncludedViewModels.Shrink(); };
+	FilterCache.Empty(FilterCache.Num());
+	ON_SCOPE_EXIT{ FilterCache.Shrink(); };
 	
 	const TSharedPtr<FSequencerEditorViewModel> RootViewModel = GetSequencer().GetViewModel();
 	const TSharedPtr<FSequencerSelection> SequencerSelection = RootViewModel ? RootViewModel->GetSelection() : nullptr;
@@ -176,25 +167,9 @@ void FSequencerTrackFilter_Selected::OnPreFilter()
 	
 	const auto ProcessSelected = [this](const FViewModelPtr& ViewModel)
 	{
-		if (IncludedViewModels.Contains(ViewModel))
+		if (!FilterCache.Contains(ViewModel))
 		{
-			return;
-		}
-			
-		// If "Foo.A" is selected, make "Foo.A.Child", etc. filtered in...
-		IncludedViewModels.Add(ViewModel, EFilterResultFlags::ApplyResultToChildren);
-	
-		for (const FViewModelPtr& Child : FParentFirstChildIterator(ViewModel))
-		{
-			// ... an if the filtered is asked whether "Foo.A.Child" is filtered in, it should also say, yes.
-			IncludedViewModels.Add(Child, EFilterResultFlags::ApplyResultToChildren);
-		}
-			
-		for (const FViewModelPtr& Parent : FParentModelIterator(ViewModel))
-		{
-			// Parent should be visible but do not tell them to propagate to all their children. 
-			// If Foo.A is selected, then Foo should be filtered in, but not Foo.B.
-			IncludedViewModels.FindOrAdd(Parent);
+			IncludeTree(FilterCache, ViewModel);
 		}
 	};
 
@@ -223,7 +198,7 @@ FString FSequencerTrackFilter_Selected::GetName() const
 	return StaticName();
 }
 
-bool FSequencerTrackFilter_Selected::PassesFilter(FSequencerTrackFilterType InItem) const
+FFilterResult FSequencerTrackFilter_Selected::Evaluate(FSequencerTrackFilterType InItem) const
 {
 	// TODO UE-371829: Remove this CVar
 	const IConsoleVariable* Variable = IConsoleManager::Get().FindConsoleVariable(TEXT("Sequencer.UseNewFilteringAlgorithm"));
@@ -233,24 +208,26 @@ bool FSequencerTrackFilter_Selected::PassesFilter(FSequencerTrackFilterType InIt
 		const TSharedPtr<FSequencerEditorViewModel> ViewModel = GetSequencer().GetViewModel();
 		if (!ViewModel.IsValid())
 		{
-			return false;
+			return EItemFilterState::Exclude;
 		}
  
 		const TSharedPtr<FSequencerSelection> SequencerSelection = ViewModel->GetSelection();
 		if (!SequencerSelection.IsValid())
 		{
-			return false;
+			return EItemFilterState::Exclude;
 		}
 		const bool bIsPartOfSequencerSelection = Algo::AnyOf(SequencerSelection->Outliner, [InItem](const FViewModelPtr& Selected)
 		{
 			return SelectionDetail::IsPartOfHierarchy(InItem, Selected);
 		});
-		return bIsPartOfSequencerSelection 
+		
+		const bool bPasses = bIsPartOfSequencerSelection 
 			// Allow external logic to make us pretend like additional items are selected without actually modifying FSequencerSelection::Outliner.
 			|| SelectionDetail::IsPartOfExtendedSelection(InItem, *this, CachedWeakExtensions);
+		return bPasses ? EItemFilterState::Include : EItemFilterState::Exclude;
 	}
 	
-	return IncludedViewModels.Contains(InItem);
+	return EvaluateFilterCache(FilterCache, InItem);
 }
 
 void FSequencerTrackFilter_Selected::BindEvents()
@@ -280,9 +257,7 @@ void FSequencerTrackFilter_Selected::UnbindEvents()
 void FSequencerTrackFilter_Selected::OnSequencerSelectionChanged()
 {
 	bNeedsPrefilterRebuild = true;
-	// TODO UE-371115: For now, we achieve this ticket behavior partially by not refiltering. However, once the ticket is completed, this should be uncommented.
-	// FilterInterface.RequestFilterUpdate();
-	// Achieved partially because any subsequent filter reevaluation, like updating the search bar and resetting it, will hide previously multi-selected items.
+	FilterInterface.RequestFilterUpdate();
 }
 
 void FSequencerTrackFilter_Selected::OnViewportSelectionChanged(UObject* const InObject)
diff --git a/Engine/Source/Editor/Sequencer/Private/Filters/Filters/SequencerTrackFilter_StickySelection.cpp b/Engine/Source/Editor/Sequencer/Private/Filters/Filters/SequencerTrackFilter_StickySelection.cpp
new file mode 100644
--- /dev/null
+++ b/Engine/Source/Editor/Sequencer/Private/Filters/Filters/SequencerTrackFilter_StickySelection.cpp
@@ -0,0 +1,111 @@
+﻿// Copyright Epic Games, Inc. All Rights Reserved.
+
+#include "SequencerTrackFilter_StickySelection.h"
+
+#include "Sequencer.h"
+#include "Filters/SequencerTrackFilterStickySelectionExtension.h"
+#include "Filters/Utils/HierarchyFilterCache.h"
+#include "MVVM/Selection/Selection.h"
+#include "MVVM/ViewModels/SequencerEditorViewModel.h"
+#include "UObject/UObjectIterator.h"
+
+namespace UE::Sequencer
+{
+namespace StickySelection
+{
+static TArray<TWeakObjectPtr<USequencerTrackFilterStickySelectionExtension>> GetSelectionExtensions()
+{
+	TArray<TWeakObjectPtr<USequencerTrackFilterStickySelectionExtension>> Extensions;
+	for (TObjectIterator<USequencerTrackFilterStickySelectionExtension> ExtensionIt(RF_NoFlags); ExtensionIt; ++ExtensionIt)
+	{
+		USequencerTrackFilterStickySelectionExtension* PotentialExtension = *ExtensionIt;
+		if (PotentialExtension
+			&& PotentialExtension->HasAnyFlags(RF_ClassDefaultObject)
+			&& !PotentialExtension->GetClass()->HasAnyClassFlags(CLASS_Deprecated | CLASS_Abstract))
+		{
+			Extensions.Emplace(PotentialExtension);
+		}
+	}
+	return Extensions;
+}
+}
+
+FSequencerTrackFilter_StickySelection::FSequencerTrackFilter_StickySelection(
+	const FSequencer& InSequencer, ISequencerTrackFilters& InFilterInterface, TSharedPtr<FFilterCategory> InCategory
+	)
+	: FSequencerTrackFilter_HierarchyBased(InFilterInterface, MoveTemp(InCategory))
+	, Sequencer(InSequencer)
+	, CachedWeakExtensions(StickySelection::GetSelectionExtensions())
+{
+	if (const TSharedPtr<FSequencerSelection> Selection = GetSequencerSelection())
+	{
+		Selection->Outliner.OnChanged.AddRaw(this, &FSequencerTrackFilter_StickySelection::OnSelectionChanged);
+	}
+}
+
+FSequencerTrackFilter_StickySelection::~FSequencerTrackFilter_StickySelection()
+{
+	if (const TSharedPtr<FSequencerSelection> Selection = GetSequencerSelection())
+	{
+		Selection->Outliner.OnChanged.RemoveAll(this);
+	}
+}
+
+FString FSequencerTrackFilter_StickySelection::GetName() const
+{
+	return StaticName();
+}
+
+void FSequencerTrackFilter_StickySelection::OnPreFilter()
+{
+	const TSharedPtr<FSequencerSelection> Selection = GetSequencerSelection();
+	if (!bNeedsPrefilterRebuild || !Selection)
+	{
+		return;
+	}
+	bNeedsPrefilterRebuild = false;
+	const bool bLastResetCacheOnUpdate = bResetCacheOnUpdate;
+	bResetCacheOnUpdate = false;
+	
+	const FOutlinerSelection& OutlinerSelection = Selection->Outliner;
+	const bool bWasCleared = OutlinerSelection.Num() == 0;
+	if (bLastResetCacheOnUpdate || bWasCleared)
+	{
+		FilterCache.Empty(FilterCache.Num());
+	}
+	
+	for (const FViewModelPtr& ViewModel : OutlinerSelection)
+	{
+		if (!FilterCache.Contains(ViewModel))
+		{
+			IncludeTree(FilterCache, ViewModel);
+		}
+	}
+	
+	FilterCache.Shrink();
+}
+
+FFilterResult FSequencerTrackFilter_StickySelection::Evaluate(FSequencerTrackFilterType InItem) const
+{
+	return EvaluateFilterCache(FilterCache, InItem, EItemFilterState::DoNotCare);
+}
+
+TSharedPtr<FSequencerSelection> FSequencerTrackFilter_StickySelection::GetSequencerSelection() const
+{
+	const TSharedPtr<FSequencerEditorViewModel> RootViewModel = GetSequencer().GetViewModel();
+	const TSharedPtr<FSequencerSelection> SequencerSelection = RootViewModel ? RootViewModel->GetSelection() : nullptr;
+	return SequencerSelection;
+}
+
+void FSequencerTrackFilter_StickySelection::OnSelectionChanged()
+{
+	bNeedsPrefilterRebuild = true;
+	
+	const auto ProcessExtension = [this](const TWeakObjectPtr<USequencerTrackFilterStickySelectionExtension>& Extension)
+	{
+		return Extension.IsValid() && Extension->IsPerformingSelectionThatShouldClearStickySelection(Sequencer);
+	};
+	const bool bIsExternalSystemClearingStickySelection = Algo::AnyOf(CachedWeakExtensions, ProcessExtension);
+	bResetCacheOnUpdate |= Sequencer.IsUpdatingSequencerSelection() || bIsExternalSystemClearingStickySelection;
+}
+} // namespace UE::Sequencer
diff --git a/Engine/Source/Editor/Sequencer/Private/Filters/Filters/SequencerTrackFilter_StickySelection.h b/Engine/Source/Editor/Sequencer/Private/Filters/Filters/SequencerTrackFilter_StickySelection.h
new file mode 100644
--- /dev/null
+++ b/Engine/Source/Editor/Sequencer/Private/Filters/Filters/SequencerTrackFilter_StickySelection.h
@@ -0,0 +1,72 @@
+﻿// Copyright Epic Games, Inc. All Rights Reserved.
+
+#pragma once
+
+#include "Filters/Filters/SequencerTrackFilter_HierarchyBased.h"
+#include "Filters/SequencerTrackFilterBase.h"
+#include "Filters/Utils/HierarchyFilterCache.h"
+#include "HAL/Platform.h"
+#include "MVVM/Selection/SequencerOutlinerSelection.h"
+
+class USequencerTrackFilterStickySelectionExtension;
+class FSequencer;
+
+namespace UE::Sequencer
+{
+class FSequencerSelection;
+
+/**
+ * Prevents clicks in the Outliner or Curve Editor from changing which items the "Selected" filter shows.
+ * Even though clicking items in the Outliner or Curve Editor changes the selection internally, the items clicked will continue to be filtered in 
+ * until the user changes the selection externally, e.g. by clicking in the viewport.
+ * This filter is designed to be used in conjunction with the "Selected" filter.
+ *
+ * The workflow that is being enabled is as follows:
+ * 1. User enables the "Selected" filter
+ * 2. User shift-clicks chest, body, and pelvis controls of a Control Rig in the viewport (-> external selection)
+ * 3. User clicks chest in the Sequencer Outliner or Curve Editor (-> internal selection)
+ * Without this filter: "Selected" filter collapses to only chest, since the internal click changed the selection.
+ * With this filter: chest, body, and pelvis remain filtered in.
+ *
+ * @see FSequencer::bUpdatingSequencerSelection for the internal/external selection distinction.
+ */
+class FSequencerTrackFilter_StickySelection : public FSequencerTrackFilter_HierarchyBased
+{
+public:
+	
+	static FString StaticName() { return TEXT("StickySelection"); }
+
+	explicit FSequencerTrackFilter_StickySelection(
+		const FSequencer& InSequencer UE_LIFETIMEBOUND, 
+		ISequencerTrackFilters& InFilterInterface, TSharedPtr<FFilterCategory> InCategory = nullptr
+		);
+	virtual ~FSequencerTrackFilter_StickySelection() override;
+	
+	//~ Begin IFilter
+	virtual FString GetName() const override;
+	virtual void OnPreFilter() override;
+	virtual FFilterResult Evaluate(FSequencerTrackFilterType InItem) const override;
+	//~ End IFilter
+	
+private:
+	
+	/** The Sequencer for which this filter was created. */
+	const FSequencer& Sequencer;
+	
+	/** Objects extending this filter. Cached at construction to speed up filtering. */
+	const TArray<TWeakObjectPtr<USequencerTrackFilterStickySelectionExtension>> CachedWeakExtensions;
+	
+	/** Whether IncludedViewModels needs to be rebuilt. */
+	bool bNeedsPrefilterRebuild = true;
+	/** Whetehr to reset the cache on update. This will be set, e.g. when the user clicks something in the viewport. */
+	bool bResetCacheOnUpdate = false;
+	/** Contains the selected items. */
+	FHierarchyFilterCache FilterCache;
+	
+	/** @return Sequencer's active selection. */
+	TSharedPtr<FSequencerSelection> GetSequencerSelection() const;
+	
+	/** Invoked when the Sequencer selection changes. */
+	void OnSelectionChanged();
+};
+} // namespace UE::Sequencer
diff --git a/Engine/Source/Editor/Sequencer/Private/Filters/SequencerFilterBar.cpp b/Engine/Source/Editor/Sequencer/Private/Filters/SequencerFilterBar.cpp
--- a/Engine/Source/Editor/Sequencer/Private/Filters/SequencerFilterBar.cpp
+++ b/Engine/Source/Editor/Sequencer/Private/Filters/SequencerFilterBar.cpp
@@ -31,6 +31,7 @@
 #include "SSequencer.h"
 #include "Algo/ForEach.h"
 #include "Filters/SequencerFilterAreaConfig.h"
+#include "Filters/SequencerTrackFilter_StickySelection.h"
 #include "Stats/StatsMisc.h"
 #include "UObject/UObjectIterator.h"
 #include "Utils/HideIsolateViewModelUtils.h"
@@ -66,6 +67,7 @@ FSequencerFilterBar::FSequencerFilterBar(FFilterBarInitArgs InArgs, FPrivateToke
 	, LevelFilter(MakeShared<FSequencerTrackFilter_Level>(*this, TransientCategory))
 	, GroupFilter(MakeShared<FSequencerTrackFilter_Group>(*this, TransientCategory))
 	, SelectedFilter(MakeShared<FSequencerTrackFilter_Selected>(*this, MiscCategory))
+	, MultiSelectFilter(MakeShared<FSequencerTrackFilter_StickySelection>(*InArgs.Sequencer, *this, TransientCategory))
 	, ModifiedFilter(MakeShared<FSequencerTrackFilter_Modified>(*this, MiscCategory))
 	, FilterData(FString())
 	, SubIdentifier(MoveTemp(InArgs.SubIdentifier))
@@ -253,6 +255,9 @@ void FSequencerFilterBar::CreateDefaultFilters()
 				return true;
 			});
 	}
+	
+	// Always active
+	SetFilterActive(MultiSelectFilter, true);
 
 	// Broadcast changes
 	CommonFilters->OnChanged().Broadcast();
@@ -1315,18 +1320,18 @@ bool FSequencerFilterBar::PassesAnyCommonFilter(const TViewModelPtr<IOutlinerExt
 	CommonFilters->ForEachFilter([this, &InNode, &bPassedAnyFilters, &bAnyFilterActive, FocusedSequence]
 		(const TSharedRef<FSequencerTrackFilter>& InFilter)
 		{
-			if (IsFilterActive(InFilter))
-			{
-				bAnyFilterActive = true;
-			
-				if (InFilter->PassesFilter(InNode))
-				{
-					bPassedAnyFilters = true;
-					return false; // Stop processing filters
-				}
-			}
-
-			return true;
+		   if (IsFilterActive(InFilter))
+		   {
+			  bAnyFilterActive = true;
+       
+			  if (InFilter->PassesFilter(InNode))
+			  {
+				 bPassedAnyFilters = true;
+				 return false; // Stop processing filters
+			  }
+		   }
+
+		   return true;
 		});
 
 	if (!bAnyFilterActive)
@@ -1719,6 +1724,12 @@ FSequencerFilterContext FSequencerFilterBar::BuildActiveFilters()
 		}
 		return true;
 	});
+	// When filtering treat MultiSelectFilter like any other common filter: any common filter needs to pass to filter in a node.
+	// But: it is not supposed to get a filter bar widget when filtering.
+	if (IsFilterActive(SelectedFilter) && IsFilterActive(MultiSelectFilter))
+	{
+		Filters.CommonFilters.Add(&MultiSelectFilter.Get());
+	}
 	
 	InternalFilters->ForEachFilter([this, &Filters](const TSharedRef<FSequencerTrackFilter>& InFilter)
 	{
diff --git a/Engine/Source/Editor/Sequencer/Private/Filters/SequencerFilterBar.h b/Engine/Source/Editor/Sequencer/Private/Filters/SequencerFilterBar.h
--- a/Engine/Source/Editor/Sequencer/Private/Filters/SequencerFilterBar.h
+++ b/Engine/Source/Editor/Sequencer/Private/Filters/SequencerFilterBar.h
@@ -29,6 +29,7 @@ enum class EFilterBarLayout : uint8;
 namespace UE::Sequencer {
 struct FSequencerFilterOverrides;
 struct FSequencerFilterContext;
+class FSequencerTrackFilter_StickySelection;
 struct FRecursiveFilterOverrides;
 class IOutlinerExtension; }
 namespace UE::Sequencer { struct FFilterBarInitArgs; }
@@ -289,6 +290,7 @@ protected:
 	TSharedRef<FSequencerTrackFilter_Level> LevelFilter;
 	TSharedRef<FSequencerTrackFilter_Group> GroupFilter;
 	TSharedRef<FSequencerTrackFilter_Selected> SelectedFilter;
+	TSharedRef<UE::Sequencer::FSequencerTrackFilter_StickySelection> MultiSelectFilter;
 	TSharedRef<FSequencerTrackFilter_Modified> ModifiedFilter;
 	
 	TArray<TSharedRef<FSequencerTrackFilter_CustomText>> CustomTextFilters;
diff --git a/Engine/Source/Editor/Sequencer/Private/Filters/Utils/HierarchyFilterCache.cpp b/Engine/Source/Editor/Sequencer/Private/Filters/Utils/HierarchyFilterCache.cpp
new file mode 100644
--- /dev/null
+++ b/Engine/Source/Editor/Sequencer/Private/Filters/Utils/HierarchyFilterCache.cpp
@@ -0,0 +1,42 @@
+﻿// Copyright Epic Games, Inc. All Rights Reserved.
+
+#include "Filters/Utils/HierarchyFilterCache.h"
+
+#include "MVVM/ViewModels/ViewModelIterators.h"
+
+namespace UE::Sequencer
+{
+void IncludeTree(FHierarchyFilterCache& InCache, const FViewModelPtr& InItem)
+{
+	// If it's already included the parents were already implicitly added before so avoid iterating it again...
+	if (!InCache.Contains(InItem))
+	{
+		// If "Foo.A" is selected, make "Foo.A.Child", etc. filtered in...
+		InCache.Add(InItem, EFilterResultFlags::ApplyResultToChildren);
+		
+		for (const FViewModelPtr& Parent : FParentModelIterator(InItem))
+		{
+			// Parent should be visible but do not tell them to propagate to all their children. 
+			// If Foo.A is selected, then Foo should be filtered in, but not Foo.B.
+			InCache.FindOrAdd(Parent);
+		}
+	}
+	
+	// ... but not all of its children might be added, yet.
+	for (const FViewModelPtr& Child : FParentFirstChildIterator(InItem))
+	{
+		// ... an if the filtered is asked whether "Foo.A.Child" is filtered in, it should also say, yes.
+		InCache.Add(Child, EFilterResultFlags::ApplyResultToChildren);
+	}
+}
+
+FFilterResult EvaluateFilterCache(const FHierarchyFilterCache& InCache, const FViewModelPtr& InItem, EItemFilterState InNotIncludedResult)
+{
+	const EFilterResultFlags* Flags = InCache.Find(InItem);
+	const bool bPassesFilter = Flags != nullptr;
+	
+	const EItemFilterState ResultState = bPassesFilter ? EItemFilterState::Include : InNotIncludedResult;
+	const EFilterResultFlags ResultFlags = bPassesFilter ? *Flags : EFilterResultFlags::ApplyResultToChildren;
+	return FFilterResult(ResultState, ResultFlags);
+}
+} // UE::Sequencer
diff --git a/Engine/Source/Editor/Sequencer/Private/Filters/Utils/HierarchyFilterCache.h b/Engine/Source/Editor/Sequencer/Private/Filters/Utils/HierarchyFilterCache.h
new file mode 100644
--- /dev/null
+++ b/Engine/Source/Editor/Sequencer/Private/Filters/Utils/HierarchyFilterCache.h
@@ -0,0 +1,30 @@
+﻿// Copyright Epic Games, Inc. All Rights Reserved.
+
+#pragma once
+
+#include "Containers/Map.h"
+#include "Filters/FilterResult.h"
+#include "MVVM/ViewModelPtr.h"
+
+#define UE_API SEQUENCER_API
+
+namespace UE::Sequencer
+{
+using FHierarchyFilterCache = TMap<TWeakViewModelPtr<FViewModel>, EFilterResultFlags>;
+
+/**
+ * Causes InItem, its parents, and all of InItem's children to pass the filter. 
+ * - Adds all InItem and InItem's children with the ApplyResultToChildren flag. 
+ * - Adds InItem's parent without the flag.
+ */
+void IncludeTree(FHierarchyFilterCache& InCache, const FViewModelPtr& InItem);
+	
+/** 
+ * @param InItem The item to evaluate
+ * @param InNotIncludedResult What to return when InItem is not in the hiearchy.
+ * @return Whether InItem passes the filter result. 
+ */
+FFilterResult EvaluateFilterCache(const FHierarchyFilterCache& InCache, const FViewModelPtr& InItem, EItemFilterState InNotIncludedResult = EItemFilterState::Exclude);
+} // namespace UE::Sequencer
+
+#undef UE_API
diff --git a/Engine/Source/Editor/Sequencer/Private/Sequencer.h b/Engine/Source/Editor/Sequencer/Private/Sequencer.h
--- a/Engine/Source/Editor/Sequencer/Private/Sequencer.h
+++ b/Engine/Source/Editor/Sequencer/Private/Sequencer.h
@@ -1308,6 +1308,9 @@ public:
 	EVisibility GetPlayRateComboVisibility() const;
 
 	virtual TSharedPtr<UE::Sequencer::SOutlinerView> GetOutlinerViewWidget() const override;
+	
+	/** @return Whether the sequencer selection is being updated from changes to the external selection. */
+	bool IsUpdatingSequencerSelection() const { return bUpdatingSequencerSelection; }
 
 private:
 
diff --git a/Engine/Source/Editor/Sequencer/Public/Filters/Filters/SequencerTrackFilter_HierarchyBased.h b/Engine/Source/Editor/Sequencer/Public/Filters/Filters/SequencerTrackFilter_HierarchyBased.h
new file mode 100644
--- /dev/null
+++ b/Engine/Source/Editor/Sequencer/Public/Filters/Filters/SequencerTrackFilter_HierarchyBased.h
@@ -0,0 +1,27 @@
+﻿// Copyright Epic Games, Inc. All Rights Reserved.
+
+#pragma once
+
+#include "Filters/SequencerTrackFilterBase.h"
+
+namespace UE::Sequencer
+{
+/** Base class for filters that intend to implement FSequencerTrackFilter::Evaluate. */
+class FSequencerTrackFilter_HierarchyBased : public FSequencerTrackFilter
+{
+public:
+	
+	FSequencerTrackFilter_HierarchyBased(ISequencerTrackFilters& InOutFilterInterface, TSharedPtr<FFilterCategory>&& InCategory)
+		: FSequencerTrackFilter(InOutFilterInterface, MoveTemp(InCategory))
+	{}
+
+	// The new algorithm does not invoke PassesFilter. This is implement simply for completeness.
+	virtual bool PassesFilter(TViewModelPtr<FViewModel> InItem) const override final
+	{
+		return PassesFilterState(Evaluate(InItem).ItemState);
+	}
+
+	// Subclasses must implement Evaluate
+	virtual FFilterResult Evaluate(FSequencerTrackFilterType InItem) const override = 0;
+};
+}
diff --git a/Engine/Source/Editor/Sequencer/Public/Filters/Filters/SequencerTrackFilter_Selected.h b/Engine/Source/Editor/Sequencer/Public/Filters/Filters/SequencerTrackFilter_Selected.h
--- a/Engine/Source/Editor/Sequencer/Public/Filters/Filters/SequencerTrackFilter_Selected.h
+++ b/Engine/Source/Editor/Sequencer/Public/Filters/Filters/SequencerTrackFilter_Selected.h
@@ -2,11 +2,11 @@
 
 #pragma once
 
-#include "Filters/SequencerTrackFilterBase.h"
+#include "SequencerTrackFilter_HierarchyBased.h"
 
 class USequencerTrackFilterSelectedExtension;
 
-class FSequencerTrackFilter_Selected : public FSequencerTrackFilter
+class FSequencerTrackFilter_Selected : public UE::Sequencer::FSequencerTrackFilter_HierarchyBased
 {
 public:
 	static FString StaticName() { return TEXT("Selected"); }
@@ -19,7 +19,6 @@ public:
 	virtual FText GetDefaultToolTipText() const override;
 	virtual TSharedPtr<FUICommandInfo> GetToggleCommand() const override;
 	virtual void ActiveStateChanged(const bool bInActive) override;
-	virtual UE::Sequencer::FFilterResult Evaluate(FSequencerTrackFilterType InItem) const override;
 	//~ End FSequencerTrackFilter
 
 	//~ Begin FFilterBase
@@ -31,7 +30,7 @@ public:
 
 	//~ Begin IFilter
 	virtual FString GetName() const override;
-	virtual bool PassesFilter(FSequencerTrackFilterType InItem) const override;
+	virtual UE::Sequencer::FFilterResult Evaluate(FSequencerTrackFilterType InItem) const override;
 	//~ End IFilter
 
 protected:
@@ -48,9 +47,8 @@ private:
 	
 	/** Whether IncludedViewModels needs to be rebuilt. */
 	bool bNeedsPrefilterRebuild = true;
-	
-	/** The view models that pass the filter. Built once during pre-filter to avoid excessive iteration of tree hierarchy. */
-	TMap<UE::Sequencer::TWeakViewModelPtr<UE::Sequencer::FViewModel>, UE::Sequencer::EFilterResultFlags> IncludedViewModels;
+	/** Contains the selected items. */
+	TMap<UE::Sequencer::TWeakViewModelPtr<UE::Sequencer::FViewModel>, UE::Sequencer::EFilterResultFlags> FilterCache;
 	
 	void BindEvents();
 	void UnbindEvents();
diff --git a/Engine/Source/Editor/Sequencer/Public/Filters/SequencerTrackFilterStickySelectionExtension.h b/Engine/Source/Editor/Sequencer/Public/Filters/SequencerTrackFilterStickySelectionExtension.h
new file mode 100644
--- /dev/null
+++ b/Engine/Source/Editor/Sequencer/Public/Filters/SequencerTrackFilterStickySelectionExtension.h
@@ -0,0 +1,18 @@
+﻿// Copyright Epic Games, Inc. All Rights Reserved.
+
+#pragma once
+
+#include "SequencerTrackFilterStickySelectionExtension.generated.h"
+
+class ISequencer;
+
+/** Derive from this class to customize FSequencerTrackFilter_StickSelection. */
+UCLASS(MinimalAPI, Abstract)
+class USequencerTrackFilterStickySelectionExtension : public UObject
+{
+	GENERATED_BODY()
+public:
+	
+	/** @return Whether this system is currenty changing selection and it should cause sticky selection to reset. */
+	virtual bool IsPerformingSelectionThatShouldClearStickySelection(const ISequencer& InSequencer) { return false; }
+};
diff --git a/Engine/Source/Runtime/LevelSequence/Private/LevelSequence.cpp b/Engine/Source/Runtime/LevelSequence/Private/LevelSequence.cpp
--- a/Engine/Source/Runtime/LevelSequence/Private/LevelSequence.cpp
+++ b/Engine/Source/Runtime/LevelSequence/Private/LevelSequence.cpp
@@ -194,6 +194,7 @@ bool ULevelSequence::IsFilterSupportedImpl(const FString& InFilterName) const
 		TEXT("SelectedAnimLayersFilter"),
 		TEXT("SelectedControlRigControl"),
 		TEXT("SkeletalMesh"),
+		TEXT("StickySelection"),
 		TEXT("Subsequence"),
 		TEXT("TimeDilation"),
 		TEXT("TimeWarp"),
