Sequencer: Add mute/disable support for decorations

Add IMutableExtension and IDeactivatableExtension to FDecorationOutlinerModel so any decoration on a track or section can be muted or disabled. Mute is transient; disable is persistent. Both states prevent entity import during compilation and skip channel proxy/section model construction in the editor.

Also fixes a subtle bug in NeedsReconstruct() where dual-interface decorations could leave stale dirty flags causing redundant reconstruction.

#Synced-CL 52466902

diff --git a/Engine/Plugins/Experimental/MovieSceneAnimMixer/Source/MovieSceneAnimMixer/Private/MovieSceneAnimationMixerTrack.cpp b/Engine/Plugins/Experimental/MovieSceneAnimMixer/Source/MovieSceneAnimMixer/Private/MovieSceneAnimationMixerTrack.cpp
--- a/Engine/Plugins/Experimental/MovieSceneAnimMixer/Source/MovieSceneAnimMixer/Private/MovieSceneAnimationMixerTrack.cpp
+++ b/Engine/Plugins/Experimental/MovieSceneAnimMixer/Source/MovieSceneAnimMixer/Private/MovieSceneAnimationMixerTrack.cpp
@@ -411,7 +411,7 @@ bool UMovieSceneAnimationMixerTrack::PopulateEvaluationTree(TMovieSceneEvaluatio
 
 	for (const UMovieSceneAnimationMixerLayer* Layer : GetLayers())
 	{
-		for (UObject* Decoration : Layer->GetDecorations())
+		for (UObject* Decoration : Layer->GetActiveDecorations())
 		{
 			if (IMovieSceneSectionProviderDecoration* SectionProvider = Cast<IMovieSceneSectionProviderDecoration>(Decoration))
 			{
@@ -1102,7 +1102,7 @@ void UMovieSceneAnimationMixerTrack::UpdateEasing()
 		{
 			continue;
 		}
-		for (TObjectPtr<UObject> Decoration : Layer.Get()->GetDecorations())
+		for (UObject* Decoration : Layer.Get()->GetActiveDecorations())
 		{
 			if (IMovieSceneSectionProviderDecoration* SectionProvider = Cast<IMovieSceneSectionProviderDecoration>(Decoration))
 			{
diff --git a/Engine/Plugins/Experimental/MovieSceneAnimMixer/Source/MovieSceneAnimMixer/Private/Systems/MovieSceneAnimMixerSystem.cpp b/Engine/Plugins/Experimental/MovieSceneAnimMixer/Source/MovieSceneAnimMixer/Private/Systems/MovieSceneAnimMixerSystem.cpp
--- a/Engine/Plugins/Experimental/MovieSceneAnimMixer/Source/MovieSceneAnimMixer/Private/Systems/MovieSceneAnimMixerSystem.cpp
+++ b/Engine/Plugins/Experimental/MovieSceneAnimMixer/Source/MovieSceneAnimMixer/Private/Systems/MovieSceneAnimMixerSystem.cpp
@@ -775,7 +775,7 @@ namespace UE::MovieScene
 		UMovieSceneEntitySystemLinker* Linker;
 		UMovieSceneAnimMixerSystem* MixerSystem;
 
-		// Layer blend task builders grouped by BoundObjectKey → {MixerLayer → Builder}
+		// Layer blend task builders grouped by BoundObjectKey -> {MixerLayer -> Builder}
 		TMap<FObjectKey, TMap<TWeakObjectPtr<UMovieSceneAnimationMixerLayer>, FLayerBlendTaskBuilder>> LayerBlendBuilders;
 
 		FEvaluateAnimMixers(TMap<FMovieSceneAnimMixerKey, TSharedPtr<FMovieSceneAnimMixer>>* InMixers, TMap<FObjectKey, TWeakPtr<FMovieSceneMixerRootMotionComponentData>>* InRootMotions, UMovieSceneEntitySystemLinker* InLinker, UMovieSceneAnimMixerSystem* InMixerSystem)
diff --git a/Engine/Source/Editor/Sequencer/Private/MVVM/ViewModels/DecorationOutlinerModel.cpp b/Engine/Source/Editor/Sequencer/Private/MVVM/ViewModels/DecorationOutlinerModel.cpp
--- a/Engine/Source/Editor/Sequencer/Private/MVVM/ViewModels/DecorationOutlinerModel.cpp
+++ b/Engine/Source/Editor/Sequencer/Private/MVVM/ViewModels/DecorationOutlinerModel.cpp
@@ -11,6 +11,8 @@
 #include "Channels/MovieSceneChannelProxy.h"
 #include "Decorations/IMovieSceneChannelDecoration.h"
 #include "Decorations/IMovieSceneSectionProviderDecoration.h"
+#include "Decorations/MovieSceneDecorationContainer.h"
+#include "MVVM/SharedViewModelData.h"
 #include "MVVM/SectionModelStorageExtension.h"
 #include "MVVM/ViewModels/CategoryModel.h"
 #include "MVVM/ViewModels/ChannelModel.h"
@@ -60,20 +62,21 @@ bool FDecorationOutlinerModel::NeedsReconstruct() const
 		return false;
 	}
 
+	// Consume all dirty flags in a single pass to avoid leaving stale flags
+	// when a decoration implements both interfaces.
+	bool bNeedsReconstruct = false;
+
 	if (IMovieSceneChannelDecoration* ChannelDecoration = Cast<IMovieSceneChannelDecoration>(Decoration))
 	{
-		if (ChannelDecoration->ConsumeDecorationChannelProxyDirty())
-		{
-			return true;
-		}
+		bNeedsReconstruct |= ChannelDecoration->ConsumeDecorationChannelProxyDirty();
 	}
 
 	if (IMovieSceneSectionProviderDecoration* SectionProviderDecoration = Cast<IMovieSceneSectionProviderDecoration>(Decoration))
 	{
-		return SectionProviderDecoration->ConsumeStructureChanged();
+		bNeedsReconstruct |= SectionProviderDecoration->ConsumeStructureChanged();
 	}
 
-	return false;
+	return bNeedsReconstruct;
 }
 
 void FDecorationOutlinerModel::OnConstruct()
@@ -575,6 +578,117 @@ UObject* FDecorationOutlinerModel::GetDecoration() const
 	return WeakDecoration.Get();
 }
 
+bool FDecorationOutlinerModel::IsMuted() const
+{
+	UMovieSceneDecorationContainerObject* Container = WeakDecorationContainer.Get();
+	UObject* Decoration = WeakDecoration.Get();
+	if (Container && Decoration)
+	{
+		return Container->IsDecorationMuted(Decoration);
+	}
+	return false;
+}
+
+void FDecorationOutlinerModel::SetIsMuted(bool bIsMuted)
+{
+	UMovieSceneDecorationContainerObject* Container = WeakDecorationContainer.Get();
+	UObject* Decoration = WeakDecoration.Get();
+	if (!Container || !Decoration)
+	{
+		return;
+	}
+
+	Container->SetDecorationMuted(Decoration, bIsMuted);
+
+	if (IMovieSceneChannelDecoration* ChannelDecoration = Cast<IMovieSceneChannelDecoration>(Decoration))
+	{
+		ChannelDecoration->InvalidateDecorationChannelProxy();
+	}
+
+	if (IMovieSceneSectionProviderDecoration* SectionProvider = Cast<IMovieSceneSectionProviderDecoration>(Decoration))
+	{
+		SectionProvider->MarkStructureChanged();
+	}
+
+	// Update the hierarchical caches so mute/deactive dimming is reflected immediately
+	if (TSharedPtr<FSharedViewModelData> CachedSharedData = GetSharedData())
+	{
+		if (FOutlinerCacheExtension* OutlinerCache = CachedSharedData->CastThis<FOutlinerCacheExtension>())
+		{
+			OutlinerCache->UpdateCachedFlags();
+		}
+	}
+
+	if (TSharedPtr<FSequencerEditorViewModel> EditorViewModel = GetEditor())
+	{
+		if (TSharedPtr<ISequencer> Sequencer = EditorViewModel->GetSequencer())
+		{
+			Sequencer->NotifyMovieSceneDataChanged(EMovieSceneDataChangeType::TrackValueChanged);
+		}
+	}
+}
+
+bool FDecorationOutlinerModel::IsDeactivated() const
+{
+	UMovieSceneDecorationContainerObject* Container = WeakDecorationContainer.Get();
+	UObject* Decoration = WeakDecoration.Get();
+	if (Container && Decoration)
+	{
+		return Container->IsDecorationDisabled(Decoration);
+	}
+	return false;
+}
+
+void FDecorationOutlinerModel::SetIsDeactivated(bool bInIsDeactivated)
+{
+	UMovieSceneDecorationContainerObject* Container = WeakDecorationContainer.Get();
+	UObject* Decoration = WeakDecoration.Get();
+	if (!Container || !Decoration)
+	{
+		return;
+	}
+
+	Container->SetDecorationDisabled(Decoration, bInIsDeactivated);
+
+	if (IMovieSceneChannelDecoration* ChannelDecoration = Cast<IMovieSceneChannelDecoration>(Decoration))
+	{
+		ChannelDecoration->InvalidateDecorationChannelProxy();
+	}
+
+	if (IMovieSceneSectionProviderDecoration* SectionProvider = Cast<IMovieSceneSectionProviderDecoration>(Decoration))
+	{
+		SectionProvider->MarkStructureChanged();
+	}
+
+	// Update the hierarchical caches so mute/deactive dimming is reflected immediately
+	if (TSharedPtr<FSharedViewModelData> CachedSharedData = GetSharedData())
+	{
+		if (FOutlinerCacheExtension* OutlinerCache = CachedSharedData->CastThis<FOutlinerCacheExtension>())
+		{
+			OutlinerCache->UpdateCachedFlags();
+		}
+	}
+
+	if (TSharedPtr<FSequencerEditorViewModel> EditorViewModel = GetEditor())
+	{
+		if (TSharedPtr<ISequencer> Sequencer = EditorViewModel->GetSequencer())
+		{
+			Sequencer->NotifyMovieSceneDataChanged(EMovieSceneDataChangeType::TrackValueChanged);
+		}
+	}
+}
+
+bool FDecorationOutlinerModel::IsDimmed() const
+{
+	UMovieSceneDecorationContainerObject* Container = WeakDecorationContainer.Get();
+	UObject* Decoration = WeakDecoration.Get();
+	if (Container && Decoration)
+	{
+		return !Container->IsDecorationActive(Decoration);
+	}
+	return false;
+}
+
 }
 
 #undef LOCTEXT_NAMESPACE
diff --git a/Engine/Source/Editor/Sequencer/Private/SSequencer.cpp b/Engine/Source/Editor/Sequencer/Private/SSequencer.cpp
--- a/Engine/Source/Editor/Sequencer/Private/SSequencer.cpp
+++ b/Engine/Source/Editor/Sequencer/Private/SSequencer.cpp
@@ -2726,8 +2726,8 @@ void SSequencer::UpdateLayoutTree()
 				}
 			}
 
-			// Pass 3: Section owners that are not tracks or track rows (e.g. FSectionOutlinerModel).
-			// These own individual sections and need per-section eval-disable for mute/solo.
+			// Pass 3: Section owners that are not tracks or track rows (e.g. FSectionOutlinerModel, FDecorationOutlinerModel).
+			// These own individual sections and need per-section eval-disable for mute/solo/deactivate.
 			for (TViewModelPtr<ISectionOwnerExtension> SectionOwner : Sequencer->GetNodeTree()->GetRootNode()->GetDescendantsOfType<ISectionOwnerExtension>())
 			{
 				// Skip nodes already handled by the track or track row loops above
@@ -2738,11 +2738,13 @@ void SSequencer::UpdateLayoutTree()
 
 				const ECachedMuteState MuteFlags = MuteState->GetCachedFlags(SectionOwner);
 				const ECachedSoloState SoloFlags = SoloState->GetCachedFlags(SectionOwner);
+				const ECachedDeactiveState DeactiveFlags = DeactiveState->GetCachedFlags(SectionOwner);
 
 				const bool bIsMuted = EnumHasAnyFlags(MuteFlags, ECachedMuteState::Muted | ECachedMuteState::ImplicitlyMutedByParent);
 				const bool bIsSoloed = EnumHasAnyFlags(SoloFlags, ECachedSoloState::Soloed | ECachedSoloState::ImplicitlySoloedByParent);
+				const bool bIsDeactivated = EnumHasAnyFlags(DeactiveFlags, ECachedDeactiveState::Deactivated | ECachedDeactiveState::ImplicitlyDeactivatedByParent);
 
-				const bool bLocalEvalDisabled = bIsSoloed ? false : (bIsMuted || (bAnySoloNodes && !bIsSoloed));
+				const bool bLocalEvalDisabled = bIsDeactivated || (bIsSoloed ? false : (bIsMuted || (bAnySoloNodes && !bIsSoloed)));
 
 				for (UMovieSceneSection* Section : SectionOwner->GetSections())
 				{
diff --git a/Engine/Source/Editor/Sequencer/Public/MVVM/ViewModels/DecorationOutlinerModel.h b/Engine/Source/Editor/Sequencer/Public/MVVM/ViewModels/DecorationOutlinerModel.h
--- a/Engine/Source/Editor/Sequencer/Public/MVVM/ViewModels/DecorationOutlinerModel.h
+++ b/Engine/Source/Editor/Sequencer/Public/MVVM/ViewModels/DecorationOutlinerModel.h
@@ -3,8 +3,10 @@
 #pragma once
 #include "ISequencerTrackEditor.h"
 #include "OutlinerItemModel.h"
+#include "MVVM/Extensions/IDeactivatableExtension.h"
 #include "MVVM/Extensions/IDeletableExtension.h"
 #include "MVVM/Extensions/ILockableExtension.h"
+#include "MVVM/Extensions/IMutableExtension.h"
 #include "MVVM/Extensions/ISectionOwnerExtension.h"
 #include "MVVM/Extensions/ITopLevelChannelHolderExtension.h"
 #include "MVVM/Extensions/ITrackAreaExtension.h"
@@ -29,6 +31,8 @@ class FDecorationOutlinerModel
 	, public ILockableExtension
 	, public IDeletableExtension
 	, public ITopLevelChannelHolderExtension
+	, public IMutableExtension
+	, public IDeactivatableExtension
 {
 public:
 	UE_SEQUENCER_DECLARE_CASTABLE_API(SEQUENCER_API, FDecorationOutlinerModel
@@ -37,7 +41,9 @@ public:
 		, ITrackAreaExtension
 		, ILockableExtension
 		, IDeletableExtension
-		, ITopLevelChannelHolderExtension)
+		, ITopLevelChannelHolderExtension
+		, IMutableExtension
+		, IDeactivatableExtension)
 
 	UE_API explicit FDecorationOutlinerModel(UMovieSceneTrack* InParentTrack, UMovieSceneDecorationContainerObject* InDecorationContainer, UObject* InDecoration);
 	UE_API virtual ~FDecorationOutlinerModel() override;
@@ -75,6 +81,17 @@ public:
 	UE_API virtual ELockableLockState GetLockState() const override;
 	UE_API virtual void SetIsLocked(bool bIsLocked) override;
 
+	/* IMutableExtension */
+	UE_API virtual bool IsMuted() const override;
+	UE_API virtual void SetIsMuted(bool bIsMuted) override;
+
+	/* IDeactivatableExtension */
+	UE_API virtual bool IsDeactivated() const override;
+	UE_API virtual void SetIsDeactivated(bool bInIsDeactivated) override;
+
+	/* IOutlinerExtension */
+	UE_API virtual bool IsDimmed() const override;
+
 	TWeakObjectPtr<UMovieSceneDecorationContainerObject> GetDecorationContainer() const { return WeakDecorationContainer; }
 
 	// Get the decoration object this model represents. Used for sidebar property display.
diff --git a/Engine/Source/Runtime/MovieScene/Private/Compilation/MovieSceneCompiledDataManager.cpp b/Engine/Source/Runtime/MovieScene/Private/Compilation/MovieSceneCompiledDataManager.cpp
--- a/Engine/Source/Runtime/MovieScene/Private/Compilation/MovieSceneCompiledDataManager.cpp
+++ b/Engine/Source/Runtime/MovieScene/Private/Compilation/MovieSceneCompiledDataManager.cpp
@@ -1525,7 +1525,7 @@ void UMovieSceneCompiledDataManager::GatherTrack(const FMovieSceneBinding* Objec
 			FieldBuilder.GetSharedMetaData().ObjectBindingID = ObjectBinding->GetObjectGuid();
 		}
 
-		for (UObject* Decoration : Track->GetDecorations())
+		for (UObject* Decoration : Track->GetActiveDecorations())
 		{
 			if (IMovieSceneEntityProvider* Provider = Cast<IMovieSceneEntityProvider>(Decoration))
 			{
diff --git a/Engine/Source/Runtime/MovieScene/Private/Decorations/MovieSceneDecorationContainer.cpp b/Engine/Source/Runtime/MovieScene/Private/Decorations/MovieSceneDecorationContainer.cpp
--- a/Engine/Source/Runtime/MovieScene/Private/Decorations/MovieSceneDecorationContainer.cpp
+++ b/Engine/Source/Runtime/MovieScene/Private/Decorations/MovieSceneDecorationContainer.cpp
@@ -81,6 +81,12 @@ void UMovieSceneDecorationContainerObject::Serialize(FArchive& Ar)
 
 	// Remove null decorations for safety
 	Decorations.RemoveNulls();
+
+	// Clean up disabled set- remove entries for decorations that no longer exist
+	if (Ar.IsLoading())
+	{
+		DisabledDecorations.Remove(nullptr);
+	}
 }
 
 void UMovieSceneDecorationContainerObject::AddDecoration(UObject* InDecoration)
@@ -103,7 +109,77 @@ UObject* UMovieSceneDecorationContainerObject::GetOrCreateDecoration(const TSubc
 void UMovieSceneDecorationContainerObject::RemoveDecoration(const TSubclassOf<UObject>& InClass)
 {
 	Decorations.RemoveDecoration(InClass, [this](UObject* Decoration){
+		MutedDecorations.Remove(Decoration);
+		DisabledDecorations.Remove(Decoration);
 		this->OnDecorationRemoved(Decoration);
 	});
 	MarkAsChanged();
 }
+
+bool UMovieSceneDecorationContainerObject::IsDecorationMuted(const UObject* Decoration) const
+{
+	return MutedDecorations.Contains(Decoration);
+}
+
+void UMovieSceneDecorationContainerObject::SetDecorationMuted(UObject* Decoration, bool bMuted)
+{
+	if (bMuted)
+	{
+		if (MutedDecorations.Add(Decoration).IsValidId())
+		{
+			MarkAsChanged();
+		}
+	}
+	else
+	{
+		if (MutedDecorations.Remove(Decoration) > 0)
+		{
+			MarkAsChanged();
+		}
+	}
+}
+
+bool UMovieSceneDecorationContainerObject::IsDecorationDisabled(const UObject* Decoration) const
+{
+	return DisabledDecorations.Contains(Decoration);
+}
+
+void UMovieSceneDecorationContainerObject::SetDecorationDisabled(UObject* Decoration, bool bDisabled)
+{
+	const bool bCurrentlyDisabled = DisabledDecorations.Contains(Decoration);
+	if (bDisabled == bCurrentlyDisabled)
+	{
+		return;
+	}
+
+	Modify();
+
+	if (bDisabled)
+	{
+		DisabledDecorations.Add(Decoration);
+	}
+	else
+	{
+		DisabledDecorations.Remove(Decoration);
+	}
+
+	MarkAsChanged();
+}
+
+bool UMovieSceneDecorationContainerObject::IsDecorationActive(const UObject* Decoration) const
+{
+	return !IsDecorationMuted(Decoration) && !IsDecorationDisabled(Decoration);
+}
+
+TArray<UObject*> UMovieSceneDecorationContainerObject::GetActiveDecorations() const
+{
+	TArray<UObject*> Result;
+	for (const TObjectPtr<UObject>& Decoration : Decorations.GetDecorations())
+	{
+		if (Decoration && IsDecorationActive(Decoration))
+		{
+			Result.Add(Decoration);
+		}
+	}
+	return Result;
+}
diff --git a/Engine/Source/Runtime/MovieScene/Public/Decorations/MovieSceneDecorationContainer.h b/Engine/Source/Runtime/MovieScene/Public/Decorations/MovieSceneDecorationContainer.h
--- a/Engine/Source/Runtime/MovieScene/Public/Decorations/MovieSceneDecorationContainer.h
+++ b/Engine/Source/Runtime/MovieScene/Public/Decorations/MovieSceneDecorationContainer.h
@@ -182,6 +182,18 @@ public:
 		GetCompatibleUserDecorationsImpl(OutClasses);
 	}
 
+	MOVIESCENE_API bool IsDecorationMuted(const UObject* Decoration) const;
+	MOVIESCENE_API void SetDecorationMuted(UObject* Decoration, bool bMuted);
+
+	MOVIESCENE_API bool IsDecorationDisabled(const UObject* Decoration) const;
+	MOVIESCENE_API void SetDecorationDisabled(UObject* Decoration, bool bDisabled);
+
+	// Returns true if the decoration is neither muted nor disabled
+	MOVIESCENE_API bool IsDecorationActive(const UObject* Decoration) const;
+
+	// Returns only decorations that are neither muted nor disabled
+	MOVIESCENE_API TArray<UObject*> GetActiveDecorations() const;
+
 protected:
 
 	virtual void OnDecorationAdded(UObject* Decoration)
@@ -202,4 +214,10 @@ protected:
 	/** Array of decorations for this movie scene */
 	UPROPERTY(EditAnywhere, Category = "Modifiers")
 	FMovieSceneDecorationContainer Decorations;
+
+	UPROPERTY(Transient)
+	TSet<TObjectPtr<UObject>> MutedDecorations;
+
+	UPROPERTY()
+	TSet<TObjectPtr<UObject>> DisabledDecorations;
 };
