[Sequencer] Simple View with toolable timeline initial release
- Shift + Left mouse drag to create key movement tool range
- Ctrl + Left mouse drag to move keys in a range
- Ctrl + Left mouse drag selected keys
- Additions to ISequencerInputHandler
- Hideable/Undockable transport controls
#Synced-CL 52513999

diff --git a/Engine/Plugins/Animation/ControlRig/Source/ControlRigEditor/Private/ControlRigEditorModule.cpp b/Engine/Plugins/Animation/ControlRig/Source/ControlRigEditor/Private/ControlRigEditorModule.cpp
--- a/Engine/Plugins/Animation/ControlRig/Source/ControlRigEditor/Private/ControlRigEditorModule.cpp
+++ b/Engine/Plugins/Animation/ControlRig/Source/ControlRigEditor/Private/ControlRigEditorModule.cpp
@@ -114,6 +114,7 @@
 #include "Editor/ControlRigNewEditor.h"
 #include "Graph/ControlRigGraphSchema.h"
 #include "Misc/LevelSequencePlaybackContext.h"
+#include "Sequencer/SequencerSimpleViewExtension.h"
 
 #if WITH_RIGVMLEGACYEDITOR
 #include "SKismetInspector.h"
@@ -361,10 +362,13 @@ void FControlRigEditorModule::StartupModule()
 		}
 	}
 
+	UE::ControlRig::SimpleView::FSequencerSimpleViewExtension::Register();
 }
 
 void FControlRigEditorModule::ShutdownModule()
 {
+	UE::ControlRig::SimpleView::FSequencerSimpleViewExtension::Unregister();
+
 	if (ICurveEditorModule* CurveEditorModule = FModuleManager::GetModulePtr<ICurveEditorModule>("CurveEditor"))
 	{
 		CurveEditorModule->UnregisterView(FControlRigSpaceChannelCurveModel::ViewID);
diff --git a/Engine/Plugins/Animation/ControlRig/Source/ControlRigEditor/Private/ControlRigEditorModule.h b/Engine/Plugins/Animation/ControlRig/Source/ControlRigEditor/Private/ControlRigEditorModule.h
--- a/Engine/Plugins/Animation/ControlRig/Source/ControlRigEditor/Private/ControlRigEditorModule.h
+++ b/Engine/Plugins/Animation/ControlRig/Source/ControlRigEditor/Private/ControlRigEditorModule.h
@@ -12,6 +12,7 @@ class UControlRigRuntimeAsset;
 class IControlRigBaseEditor;
 class UBlueprint;
 class IAssetTypeActions;
+class ISequencer;
 class UMaterial;
 class UAnimSequence;
 class USkeletalMesh;
@@ -63,8 +64,8 @@ private:
 	//property sections
 	TSharedRef<FPropertySection> RegisterPropertySection(FPropertyEditorModule& PropertyModule, FName ClassName, FName SectionName, FText DisplayName);
 	void UnregisterPropertySectionMappings();
-	TMultiMap<FName, FName> RegisteredPropertySections;
 
+	TMultiMap<FName, FName> RegisteredPropertySections;
 
 	/** Handle for our sequencer control rig parameter track editor */
 	FDelegateHandle ControlRigParameterTrackCreateEditorHandle;
@@ -87,4 +88,4 @@ private:
 	FGuid SpaceCustomChannelGuid;
 	FGuid ConstraintCustomChannelGuid;
 
-};
\ No newline at end of file
+};
diff --git a/Engine/Plugins/Animation/MLDeformer/MLDeformerFramework/Source/MLDeformerFrameworkEditor/Private/SMLDeformerTimeline.cpp b/Engine/Plugins/Animation/MLDeformer/MLDeformerFramework/Source/MLDeformerFrameworkEditor/Private/SMLDeformerTimeline.cpp
--- a/Engine/Plugins/Animation/MLDeformer/MLDeformerFramework/Source/MLDeformerFrameworkEditor/Private/SMLDeformerTimeline.cpp
+++ b/Engine/Plugins/Animation/MLDeformer/MLDeformerFramework/Source/MLDeformerFrameworkEditor/Private/SMLDeformerTimeline.cpp
@@ -78,7 +78,7 @@ namespace UE::MLDeformer
 		virtual FReply OnMouseButtonUp(SWidget& WidgetOwner, const FGeometry& MyGeometry, const FPointerEvent& MouseEvent) override;
 		virtual FReply OnMouseMove(SWidget& WidgetOwner, const FGeometry& MyGeometry, const FPointerEvent& MouseEvent) override;
 		virtual FReply OnMouseWheel(SWidget& WidgetOwner, const FGeometry& MyGeometry, const FPointerEvent& MouseEvent) override;
-		virtual FCursorReply OnCursorQuery(TSharedRef<const SWidget> WidgetOwner, const FGeometry& MyGeometry, const FPointerEvent& CursorEvent) const override;
+		virtual FCursorReply OnCursorQuery(const SWidget& WidgetOwner, const FGeometry& MyGeometry, const FPointerEvent& CursorEvent) const override;
 		virtual void SetViewRange(double NewRangeMin, double NewRangeMax, EViewRangeInterpolation Interpolation) override;
 		virtual void SetClampRange(double NewRangeMin, double NewRangeMax) override;
 		virtual void SetPlayRange(FFrameNumber RangeStart, int32 RangeDuration) override;
@@ -877,7 +877,7 @@ namespace UE::MLDeformer
 		return FReply::Unhandled();
 	}
 
-	FCursorReply FMLTimeSliderController::OnCursorQuery(TSharedRef<const SWidget> WidgetOwner, const FGeometry& MyGeometry, const FPointerEvent& CursorEvent) const
+	FCursorReply FMLTimeSliderController::OnCursorQuery(const SWidget& WidgetOwner, const FGeometry& MyGeometry, const FPointerEvent& CursorEvent) const
 	{
 		const FScrubRangeToScreen RangeToScreen(TimeSliderArgs.ViewRange.Get(), MyGeometry.Size);
 		const FFrameRate FrameResolution = GetTickResolution();
diff --git a/Engine/Plugins/Experimental/Animation/AnimDatabase/Source/AnimDatabaseEditor/Private/SAnimDatabaseEditorTimeline.cpp b/Engine/Plugins/Experimental/Animation/AnimDatabase/Source/AnimDatabaseEditor/Private/SAnimDatabaseEditorTimeline.cpp
--- a/Engine/Plugins/Experimental/Animation/AnimDatabase/Source/AnimDatabaseEditor/Private/SAnimDatabaseEditorTimeline.cpp
+++ b/Engine/Plugins/Experimental/Animation/AnimDatabase/Source/AnimDatabaseEditor/Private/SAnimDatabaseEditorTimeline.cpp
@@ -184,7 +184,7 @@ namespace UE::AnimDatabase::Editor
 		virtual FReply OnMouseButtonUp(SWidget& WidgetOwner, const FGeometry& MyGeometry, const FPointerEvent& MouseEvent) override;
 		virtual FReply OnMouseMove(SWidget& WidgetOwner, const FGeometry& MyGeometry, const FPointerEvent& MouseEvent) override;
 		virtual FReply OnMouseWheel(SWidget& WidgetOwner, const FGeometry& MyGeometry, const FPointerEvent& MouseEvent) override;
-		virtual FCursorReply OnCursorQuery(TSharedRef<const SWidget> WidgetOwner, const FGeometry& MyGeometry, const FPointerEvent& CursorEvent) const override;
+		virtual FCursorReply OnCursorQuery(const SWidget& OwnerWidget, const FGeometry& MyGeometry, const FPointerEvent& CursorEvent) const override;
 		virtual void SetViewRange(double NewRangeMin, double NewRangeMax, EViewRangeInterpolation Interpolation) override;
 		virtual void SetClampRange(double NewRangeMin, double NewRangeMax) override;
 		virtual void SetPlayRange(FFrameNumber RangeStart, int32 RangeDuration) override;
@@ -1200,7 +1200,7 @@ namespace UE::AnimDatabase::Editor
 			TSharedPtr<FTimeSliderController> TimeSliderController = WeakTimeSliderController.Pin();
 			if (TimeSliderController.IsValid())
 			{
-				return TimeSliderController->OnCursorQuery(SharedThis(this), MyGeometry, CursorEvent);
+				return TimeSliderController->OnCursorQuery(*this, MyGeometry, CursorEvent);
 			}
 		}
 
@@ -2235,7 +2235,7 @@ namespace UE::AnimDatabase::Editor
 		return FReply::Unhandled();
 	}
 
-	FCursorReply FTimeSliderController::OnCursorQuery(TSharedRef<const SWidget> WidgetOwner, const FGeometry& MyGeometry, const FPointerEvent& CursorEvent) const
+	FCursorReply FTimeSliderController::OnCursorQuery(const SWidget& OwnerWidget, const FGeometry& MyGeometry, const FPointerEvent& CursorEvent) const
 	{
 		const FScrubRangeToScreen RangeToScreen(TimeSliderArgs.ViewRange.Get(), MyGeometry.Size);
 		const FFrameRate FrameResolution = GetTickResolution();
diff --git a/Engine/Plugins/Experimental/Dataflow/Source/DataflowEditor/Private/Widgets/DataflowTimeSliderController.cpp b/Engine/Plugins/Experimental/Dataflow/Source/DataflowEditor/Private/Widgets/DataflowTimeSliderController.cpp
--- a/Engine/Plugins/Experimental/Dataflow/Source/DataflowEditor/Private/Widgets/DataflowTimeSliderController.cpp
+++ b/Engine/Plugins/Experimental/Dataflow/Source/DataflowEditor/Private/Widgets/DataflowTimeSliderController.cpp
@@ -510,7 +510,7 @@ FReply FDataflowTimeSliderController::OnMouseWheel(SWidget& WidgetOwner, const F
 	return FReply::Unhandled();
 }
 
-FCursorReply FDataflowTimeSliderController::OnCursorQuery(TSharedRef<const SWidget> WidgetOwner, const FGeometry& MyGeometry, const FPointerEvent& CursorEvent) const
+FCursorReply FDataflowTimeSliderController::OnCursorQuery(const SWidget& WidgetOwner, const FGeometry& MyGeometry, const FPointerEvent& CursorEvent) const
 {
 	const FScrubRangeToScreen RangeToScreen(TimeSliderArgs.ViewRange.Get(), MyGeometry.Size);
 
diff --git a/Engine/Plugins/Experimental/Dataflow/Source/DataflowEditor/Private/Widgets/DataflowTimeSliderController.h b/Engine/Plugins/Experimental/Dataflow/Source/DataflowEditor/Private/Widgets/DataflowTimeSliderController.h
--- a/Engine/Plugins/Experimental/Dataflow/Source/DataflowEditor/Private/Widgets/DataflowTimeSliderController.h
+++ b/Engine/Plugins/Experimental/Dataflow/Source/DataflowEditor/Private/Widgets/DataflowTimeSliderController.h
@@ -29,7 +29,7 @@ public:
 	virtual FReply OnMouseButtonDoubleClick(SWidget& OwnerWidget, const FGeometry& MyGeometry, const FPointerEvent& MouseEvent) override;
 	virtual FReply OnMouseMove(SWidget& WidgetOwner, const FGeometry& MyGeometry, const FPointerEvent& MouseEvent) override;
 	virtual FReply OnMouseWheel(SWidget& WidgetOwner, const FGeometry& MyGeometry, const FPointerEvent& MouseEvent) override;
-	virtual FCursorReply OnCursorQuery(TSharedRef<const SWidget> WidgetOwner, const FGeometry& MyGeometry, const FPointerEvent& CursorEvent) const override;
+	virtual FCursorReply OnCursorQuery(const SWidget& WidgetOwner, const FGeometry& MyGeometry, const FPointerEvent& CursorEvent) const override;
 	virtual void SetViewRange(double NewRangeMin, double NewRangeMax, EViewRangeInterpolation Interpolation) override;
 	virtual void SetClampRange(double NewRangeMin, double NewRangeMax) override;
 	virtual void SetPlayRange(FFrameNumber RangeStart, int32 RangeDuration) override;
diff --git a/Engine/Plugins/Runtime/GeometryCache/Source/GeometryCacheEd/Private/GeometryCacheTimeSliderController.cpp b/Engine/Plugins/Runtime/GeometryCache/Source/GeometryCacheEd/Private/GeometryCacheTimeSliderController.cpp
--- a/Engine/Plugins/Runtime/GeometryCache/Source/GeometryCacheEd/Private/GeometryCacheTimeSliderController.cpp
+++ b/Engine/Plugins/Runtime/GeometryCache/Source/GeometryCacheEd/Private/GeometryCacheTimeSliderController.cpp
@@ -502,7 +502,7 @@ FReply FGeometryCacheTimeSlideController::OnMouseWheel(SWidget& WidgetOwner, con
 	return FReply::Unhandled();
 }
 
-FCursorReply FGeometryCacheTimeSlideController::OnCursorQuery(TSharedRef<const SWidget> WidgetOwner, const FGeometry& MyGeometry, const FPointerEvent& CursorEvent) const
+FCursorReply FGeometryCacheTimeSlideController::OnCursorQuery(const SWidget& WidgetOwner, const FGeometry& MyGeometry, const FPointerEvent& CursorEvent) const
 {
 	const FScrubRangeToScreen RangeToScreen(TimeSliderArgs.ViewRange.Get(), MyGeometry.Size);
 
diff --git a/Engine/Plugins/Runtime/GeometryCache/Source/GeometryCacheEd/Private/GeometryCacheTimeSliderController.h b/Engine/Plugins/Runtime/GeometryCache/Source/GeometryCacheEd/Private/GeometryCacheTimeSliderController.h
--- a/Engine/Plugins/Runtime/GeometryCache/Source/GeometryCacheEd/Private/GeometryCacheTimeSliderController.h
+++ b/Engine/Plugins/Runtime/GeometryCache/Source/GeometryCacheEd/Private/GeometryCacheTimeSliderController.h
@@ -29,7 +29,7 @@ public:
 	virtual FReply OnMouseButtonDoubleClick(SWidget& OwnerWidget, const FGeometry& MyGeometry, const FPointerEvent& MouseEvent) override;
 	virtual FReply OnMouseMove(SWidget& WidgetOwner, const FGeometry& MyGeometry, const FPointerEvent& MouseEvent) override;
 	virtual FReply OnMouseWheel(SWidget& WidgetOwner, const FGeometry& MyGeometry, const FPointerEvent& MouseEvent) override;
-	virtual FCursorReply OnCursorQuery(TSharedRef<const SWidget> WidgetOwner, const FGeometry& MyGeometry, const FPointerEvent& CursorEvent) const override;
+	virtual FCursorReply OnCursorQuery(const SWidget& WidgetOwner, const FGeometry& MyGeometry, const FPointerEvent& CursorEvent) const override;
 	virtual void SetViewRange(double NewRangeMin, double NewRangeMax, EViewRangeInterpolation Interpolation) override;
 	virtual void SetClampRange(double NewRangeMin, double NewRangeMax) override;
 	virtual void SetPlayRange(FFrameNumber RangeStart, int32 RangeDuration) override;
diff --git a/Engine/Plugins/VirtualProduction/CameraCalibration/Source/CameraCalibrationEditor/Private/AssetEditor/CameraCalibrationTimeSliderController.h b/Engine/Plugins/VirtualProduction/CameraCalibration/Source/CameraCalibrationEditor/Private/AssetEditor/CameraCalibrationTimeSliderController.h
--- a/Engine/Plugins/VirtualProduction/CameraCalibration/Source/CameraCalibrationEditor/Private/AssetEditor/CameraCalibrationTimeSliderController.h
+++ b/Engine/Plugins/VirtualProduction/CameraCalibration/Source/CameraCalibrationEditor/Private/AssetEditor/CameraCalibrationTimeSliderController.h
@@ -29,7 +29,7 @@ public:
 	virtual int32 OnPaintViewArea( const FGeometry& AllottedGeometry, const FSlateRect& MyCullingRect, FSlateWindowElementList& OutDrawElements, int32 LayerId, bool bEnabled, const FPaintViewAreaArgs& Args ) const override;
 	virtual void SetViewRange( double NewRangeMin, double NewRangeMax, EViewRangeInterpolation Interpolation ) override;
 
-	virtual FCursorReply OnCursorQuery( TSharedRef<const SWidget> WidgetOwner, const FGeometry& MyGeometry, const FPointerEvent& CursorEvent ) const override { return FCursorReply::Unhandled(); }
+	virtual FCursorReply OnCursorQuery( const SWidget& OwnerWidget, const FGeometry& MyGeometry, const FPointerEvent& CursorEvent ) const override { return FCursorReply::Unhandled(); }
 	virtual FFrameRate GetDisplayRate() const override {  return TimeSliderArgs.DisplayRate.Get(); }
 	virtual FFrameRate GetTickResolution() const override { return TimeSliderArgs.TickResolution.Get(); }
 	virtual FFrameTime GetScrubPosition() const override { return TimeSliderArgs.ScrubPosition.Get(); }
diff --git a/Engine/Source/Editor/Persona/Private/AnimTimeline/AnimTimeSliderController.cpp b/Engine/Source/Editor/Persona/Private/AnimTimeline/AnimTimeSliderController.cpp
--- a/Engine/Source/Editor/Persona/Private/AnimTimeline/AnimTimeSliderController.cpp
+++ b/Engine/Source/Editor/Persona/Private/AnimTimeline/AnimTimeSliderController.cpp
@@ -884,7 +884,7 @@ FReply FAnimTimeSliderController::OnMouseWheel( SWidget& WidgetOwner, const FGeo
 	return FReply::Unhandled();
 }
 
-FCursorReply FAnimTimeSliderController::OnCursorQuery( TSharedRef<const SWidget> WidgetOwner, const FGeometry& MyGeometry, const FPointerEvent& CursorEvent ) const
+FCursorReply FAnimTimeSliderController::OnCursorQuery( const SWidget& WidgetOwner, const FGeometry& MyGeometry, const FPointerEvent& CursorEvent ) const
 {
 	const FScrubRangeToScreen RangeToScreen(TimeSliderArgs.ViewRange.Get(), MyGeometry.Size);
 
diff --git a/Engine/Source/Editor/Persona/Private/AnimTimeline/AnimTimeSliderController.h b/Engine/Source/Editor/Persona/Private/AnimTimeline/AnimTimeSliderController.h
--- a/Engine/Source/Editor/Persona/Private/AnimTimeline/AnimTimeSliderController.h
+++ b/Engine/Source/Editor/Persona/Private/AnimTimeline/AnimTimeSliderController.h
@@ -56,7 +56,7 @@ public:
 	virtual FReply OnMouseButtonDoubleClick(SWidget& OwnerWidget, const FGeometry& MyGeometry, const FPointerEvent& MouseEvent) override;
 	virtual FReply OnMouseMove( SWidget& WidgetOwner, const FGeometry& MyGeometry, const FPointerEvent& MouseEvent ) override;
 	virtual FReply OnMouseWheel( SWidget& WidgetOwner, const FGeometry& MyGeometry, const FPointerEvent& MouseEvent ) override;
-	virtual FCursorReply OnCursorQuery( TSharedRef<const SWidget> WidgetOwner, const FGeometry& MyGeometry, const FPointerEvent& CursorEvent ) const override;
+	virtual FCursorReply OnCursorQuery( const SWidget& WidgetOwner, const FGeometry& MyGeometry, const FPointerEvent& CursorEvent ) const override;
 	virtual void SetViewRange( double NewRangeMin, double NewRangeMax, EViewRangeInterpolation Interpolation ) override;
 	virtual void SetClampRange( double NewRangeMin, double NewRangeMax ) override;
 	virtual void SetPlayRange( FFrameNumber RangeStart, int32 RangeDuration ) override;
diff --git a/Engine/Source/Editor/Persona/Private/AnimTimeline/SAnimTrackArea.cpp b/Engine/Source/Editor/Persona/Private/AnimTimeline/SAnimTrackArea.cpp
--- a/Engine/Source/Editor/Persona/Private/AnimTimeline/SAnimTrackArea.cpp
+++ b/Engine/Source/Editor/Persona/Private/AnimTimeline/SAnimTrackArea.cpp
@@ -223,7 +223,7 @@ FCursorReply SAnimTrackArea::OnCursorQuery( const FGeometry& MyGeometry, const F
 		TSharedPtr<FAnimTimeSliderController> TimeSliderController = WeakTimeSliderController.Pin();
 		if(TimeSliderController.IsValid())
 		{
-			return TimeSliderController->OnCursorQuery(SharedThis(this), MyGeometry, CursorEvent);
+			return TimeSliderController->OnCursorQuery(*this, MyGeometry, CursorEvent);
 		}
 	}
 
diff --git a/Engine/Source/Editor/Sequencer/Private/CurveEditor/Widgets/Tree/SCurveTreeContent.cpp b/Engine/Source/Editor/Sequencer/Private/CurveEditor/Widgets/Tree/SCurveTreeContent.cpp
--- a/Engine/Source/Editor/Sequencer/Private/CurveEditor/Widgets/Tree/SCurveTreeContent.cpp
+++ b/Engine/Source/Editor/Sequencer/Private/CurveEditor/Widgets/Tree/SCurveTreeContent.cpp
@@ -138,7 +138,7 @@ TSharedRef<SWidget> SCurveTreeContent::MakeFooter(const FArguments& InArgs)
 	return SNew(SBorder)
 		.BorderImage(FAppStyle::GetBrush("ToolPanel.GroupBorder"))
 		.Clipping(EWidgetClipping::ClipToBounds)
-		.Visibility_Static(&GetPopoutTransportControlsVisibility_CurveEditor, StaticCastWeakPtr<ISequencer>(Sequencer.ToWeakPtr()))
+		.Visibility(this, &SCurveTreeContent::GetTransportControlsVisibility)
 		[
 			SNew(SHorizontalBox)
 
@@ -232,6 +232,17 @@ TSet<TWeakViewModelPtr<IOutlinerExtension>> SCurveTreeContent::GetSelectedItems(
 		}
 	}
 	return Selected;
-	
 }
+
+EVisibility SCurveTreeContent::GetTransportControlsVisibility() const
+{
+	if (const TSharedPtr<FSequencer> Sequencer = WeakSequencer.Pin())
+	{
+		const bool bShowTransportControls = GetDefault<UCurveEditorSettings>()->GetShowTransportControls();
+		return IsPopoutTransportControlsVisible(Sequencer.ToSharedRef(), bShowTransportControls)
+			? EVisibility::Visible : EVisibility::Collapsed;
+	}
+	return EVisibility::Visible;
+}
+
 } // namespace UE::Sequencer
diff --git a/Engine/Source/Editor/Sequencer/Private/CurveEditor/Widgets/Tree/SOldCurveTreeContent.cpp b/Engine/Source/Editor/Sequencer/Private/CurveEditor/Widgets/Tree/SOldCurveTreeContent.cpp
--- a/Engine/Source/Editor/Sequencer/Private/CurveEditor/Widgets/Tree/SOldCurveTreeContent.cpp
+++ b/Engine/Source/Editor/Sequencer/Private/CurveEditor/Widgets/Tree/SOldCurveTreeContent.cpp
@@ -115,7 +115,7 @@ TSharedRef<SWidget> SOldCurveTreeContent::MakeFooter(const FArguments& InArgs, c
 	return SNew(SBorder)
 		.BorderImage(FAppStyle::GetBrush("ToolPanel.GroupBorder"))
 		.Clipping(EWidgetClipping::ClipToBounds)
-		.Visibility_Static(&UE::Sequencer::GetPopoutTransportControlsVisibility_CurveEditor, StaticCastWeakPtr<ISequencer>(InSequencer.ToWeakPtr()))
+		.Visibility(this, &SOldCurveTreeContent::GetTransportControlsVisibility)
 		[
 			SNew(SHorizontalBox)
 
@@ -179,6 +179,18 @@ USequencerSettings* SOldCurveTreeContent::GetSequencerSettings() const
 	const TSharedPtr<ISequencer> Sequencer = WeakSequencer.Pin();
 	return Sequencer.IsValid() ? Sequencer->GetSequencerSettings() : nullptr;
 }
+
+EVisibility SOldCurveTreeContent::GetTransportControlsVisibility() const
+{
+	if (const TSharedPtr<FSequencer> Sequencer = WeakSequencer.Pin())
+	{
+		const bool bShowTransportControls = GetDefault<UCurveEditorSettings>()->GetShowTransportControls();
+		return IsPopoutTransportControlsVisible(Sequencer.ToSharedRef(), bShowTransportControls)
+			? EVisibility::Visible : EVisibility::Collapsed;
+	}
+	return EVisibility::Visible;
+}
+
 } // namespace UE::Sequencer
 
 #undef LOCTEXT_NAMESPACE
\ No newline at end of file
diff --git a/Engine/Source/Editor/Sequencer/Private/CurveEditor/Widgets/Tree/SOldCurveTreeContent.h b/Engine/Source/Editor/Sequencer/Private/CurveEditor/Widgets/Tree/SOldCurveTreeContent.h
--- a/Engine/Source/Editor/Sequencer/Private/CurveEditor/Widgets/Tree/SOldCurveTreeContent.h
+++ b/Engine/Source/Editor/Sequencer/Private/CurveEditor/Widgets/Tree/SOldCurveTreeContent.h
@@ -62,5 +62,7 @@ private:
 	TSharedRef<SWidget> MakeFooter(const FArguments& InArgs, const TSharedRef<FSequencer> InSequencer);
 	
 	USequencerSettings* GetSequencerSettings() const;
+
+	EVisibility GetTransportControlsVisibility() const;
 };
 } // namespace UE::Sequencer
diff --git a/Engine/Source/Editor/Sequencer/Private/MVVM/CurveEditorExtension.cpp b/Engine/Source/Editor/Sequencer/Private/MVVM/CurveEditorExtension.cpp
--- a/Engine/Source/Editor/Sequencer/Private/MVVM/CurveEditorExtension.cpp
+++ b/Engine/Source/Editor/Sequencer/Private/MVVM/CurveEditorExtension.cpp
@@ -6,8 +6,12 @@
 #include "Filters/Linking/Mode/ILinkedFilterViewModel.h"
 #include "FrameNumberDetailsCustomization.h"
 #include "Framework/Docking/TabManager.h"
-#include "MVVM/ViewModels/SequencerEditorViewModel.h"
+#include "Menus/SequencerToolbarUtils.h"
+#include "Misc/TransportControlsHelper.h"
 #include "Modification/Resolution/CurveMetaDataIdentifiers.h"
+#include "Modification/Resolution/CurveModelLookUpInfo.h"
+#include "Modification/Utils/ScopedSelectionChange.h"
+#include "MVVM/ViewModels/SequencerEditorViewModel.h"
 #include "SCurveEditorPanel.h"
 #include "SCurveKeyDetailPanel.h"
 #include "Sequencer.h"
@@ -130,6 +134,20 @@ void FCurveEditorExtension::RequestSyncSelection()
 }
 
 
+EVisibility FCurveEditorExtension::GetPopoutTransportControlsVisibility() const
+{
+	const TSharedPtr<ISequencer> Sequencer = GetSequencer();
+	if (!Sequencer.IsValid())
+	{
+		return EVisibility::Collapsed;
+	}
+
+	const bool bShowTransportControls = GetDefault<UCurveEditorSettings>()->GetShowTransportControls();
+
+	return IsPopoutTransportControlsVisible(Sequencer.ToSharedRef(), bShowTransportControls)
+		? EVisibility::Visible : EVisibility::Collapsed;
+}
+
 } // namespace Sequencer
 } // namespace UE
 
diff --git a/Engine/Source/Editor/Sequencer/Private/MVVM/Views/SChannelView.cpp b/Engine/Source/Editor/Sequencer/Private/MVVM/Views/SChannelView.cpp
--- a/Engine/Source/Editor/Sequencer/Private/MVVM/Views/SChannelView.cpp
+++ b/Engine/Source/Editor/Sequencer/Private/MVVM/Views/SChannelView.cpp
@@ -550,7 +550,7 @@ int32 SChannelView::OnPaint(const FPaintArgs& Args, const FGeometry& AllottedGeo
 	PaintArgs.DrawElements = &OutDrawElements;
 	PaintArgs.KeyThrobValue = SSequencerSection::GetKeySelectionThrobValue();
 	PaintArgs.DrawEffects = bParentEnabled ? ESlateDrawEffect::None : ESlateDrawEffect::DisabledEffect;
-	PaintArgs.SelectionColor = FAppStyle::GetSlateColor(SelectionColorName).GetColor(InWidgetStyle);
+	PaintArgs.KeySelectionFillColor = FAppStyle::GetSlateColor(SelectionColorName).GetColor(InWidgetStyle);
 	if (KeyBarColor.IsSet())
 	{
 		PaintArgs.KeyBarColor = KeyBarColor.Get();
diff --git a/Engine/Source/Editor/Sequencer/Private/Misc/TransportControlsHelper.cpp b/Engine/Source/Editor/Sequencer/Private/Misc/TransportControlsHelper.cpp
--- a/Engine/Source/Editor/Sequencer/Private/Misc/TransportControlsHelper.cpp
+++ b/Engine/Source/Editor/Sequencer/Private/Misc/TransportControlsHelper.cpp
@@ -1,8 +1,6 @@
-﻿// Copyright Epic Games, Inc. All Rights Reserved.
+// Copyright Epic Games, Inc. All Rights Reserved.
 
 #include "TransportControlsHelper.h"
-
-#include "CurveEditorSettings.h"
 #include "Framework/Docking/TabManager.h"
 #include "ISequencer.h"
 #include "LevelSequence.h"
@@ -38,6 +36,7 @@ TSharedRef<SPopoutTabInlineContent> MakePopoutTransportControls(const TSharedRef
 		.TabGroup(WorkspaceMenu::GetMenuStructure().GetLevelEditorCinematicsCategory())
 		.AutosizeWhenFloating(true)
 		.HideDockTabStackTab(true)
+		.ContentHAlign(HAlign_Fill)
 		.InitiallyPoppedOut(SequencerSettings->AreTransportControlsPoppedOut())
 		.OnPopoutStateChanged_Lambda([WeakSequencer](const bool bPoppedOut)
 			{
@@ -50,33 +49,33 @@ TSharedRef<SPopoutTabInlineContent> MakePopoutTransportControls(const TSharedRef
 				}
 			})
 		[
-			Content.ToSharedRef()
+			SNew(SBox)
+			.HAlign(HAlign_Center)
+			.VAlign(VAlign_Center)
+			[
+				Content.ToSharedRef()
+			]
 		];
 }
 
-EVisibility GetPopoutTransportControlsVisibility_CurveEditor(const TSharedRef<ISequencer>& InSequencer)
+bool IsPopoutTransportControlsVisible(const TSharedRef<ISequencer>& InSequencer
+	, const bool bWantsToShowTransportControls)
 {
+	// Only support level sequences for now. Always show transport controls otherwise.
 	const UMovieSceneSequence* const RootSequence = InSequencer->GetRootMovieSceneSequence();
 	if (!RootSequence || !RootSequence->IsA<ULevelSequence>())
 	{
-		return EVisibility::Visible;
+		return true;
 	}
-	
+
+	// Sequencer transport popped out state overrides curve editor transport controls visibility
 	const USequencerSettings* const SequencerSettings = InSequencer->GetSequencerSettings();
-	if (!SequencerSettings)
+	if (!SequencerSettings || SequencerSettings->AreTransportControlsPoppedOut())
 	{
-		return EVisibility::Visible;
+		return false;
 	}
-	
-	const bool bShowTransportControls = GetDefault<UCurveEditorSettings>()->GetShowTransportControls();
-	const bool bTransportControlsPoppedOut = SequencerSettings->AreTransportControlsPoppedOut();
-	return (!bShowTransportControls || bTransportControlsPoppedOut) ? EVisibility::Collapsed : EVisibility::Visible;
-}
 
-EVisibility GetPopoutTransportControlsVisibility_CurveEditor(TWeakPtr<ISequencer> InWeakSequencer)
-{
-	const TSharedPtr<ISequencer> Sequencer = InWeakSequencer.Pin();
-	return Sequencer ? GetPopoutTransportControlsVisibility_CurveEditor(Sequencer.ToSharedRef()) : EVisibility::Collapsed;
+	return bWantsToShowTransportControls;
 }
 
 } // namespace UE::Sequencer
diff --git a/Engine/Source/Editor/Sequencer/Private/Misc/TransportControlsHelper.h b/Engine/Source/Editor/Sequencer/Private/Misc/TransportControlsHelper.h
--- a/Engine/Source/Editor/Sequencer/Private/Misc/TransportControlsHelper.h
+++ b/Engine/Source/Editor/Sequencer/Private/Misc/TransportControlsHelper.h
@@ -20,8 +20,14 @@ namespace UE::Sequencer
 TSharedRef<SPopoutTabInlineContent> MakePopoutTransportControls(const TSharedRef<ISequencer>& InSequencer
 	, const TSharedPtr<FTabManager>& InTabManager, const TSharedPtr<SWidget>& InOverriddenContent = nullptr);
 
-/** @return Visibility the popout transport controls should have in Sequencer's curve editor: must be configured to be visible and not popped out. */
-EVisibility GetPopoutTransportControlsVisibility_CurveEditor(const TSharedRef<ISequencer>& InSequencer);
-/** Overload useful for shorthand delegate construction. */
-EVisibility GetPopoutTransportControlsVisibility_CurveEditor(TWeakPtr<ISequencer> InWeakSequencer);
+/**
+ * Determines if a transport control that can be popped out is visible in its inline content, or hidden 
+ * due to being popped out or bWantsToShowTransportControls is false. This function is set up like this to
+ * enable its use with other systems that want a transport control that can be popped out with its own
+ * visibility for the transport controls. The popped out state overrides visibility everywhere, but when not
+ * popped out, visibility should be taken from the system that defines the inline content.
+ */
+bool IsPopoutTransportControlsVisible(const TSharedRef<ISequencer>& InSequencer
+	, const bool bWantsToShowTransportControls);
+
 } // namespace UE::Sequencer
diff --git a/Engine/Source/Editor/Sequencer/Private/SSequencer.cpp b/Engine/Source/Editor/Sequencer/Private/SSequencer.cpp
--- a/Engine/Source/Editor/Sequencer/Private/SSequencer.cpp
+++ b/Engine/Source/Editor/Sequencer/Private/SSequencer.cpp
@@ -56,6 +56,7 @@
 #include "Widgets/Input/SComboButton.h"
 #include "Widgets/Layout/SSeparator.h"
 #include "Widgets/Layout/SSplitter.h"
+#include "Widgets/Layout/SWidgetSwitcher.h"
 #include "Widgets/Input/SSpinBox.h"
 #include "Widgets/Input/SCheckBox.h"
 #include "Widgets/Docking/SDockTab.h"
@@ -153,9 +154,12 @@
 #include "MVVM/Extensions/ITopTimeSliderOverlayExtension.h"
 #include "MVVM/Extensions/ITrackAreaOverlayExtension.h"
 #include "Settings/EditorExperimentalSettings.h"
+#include "SimpleView/SimpleViewTimeline.h"
+#include "SimpleView/SimpleViewUtils.h"
 #include "SPopoutTabInlineContent.h"
+#include "ToolableTimeline/ToolableTimeline.h"
 #include "Widgets/Docking/SDockTab.h"
-#include "Widgets/Layout/SScaleBox.h"
+#include "Widgets/Layout/SScrollBox.h"
 #include "Widgets/Sidebar/SequencerSelectionDrawer.h"
 #include "WorkspaceMenuStructure.h"
 #include "WorkspaceMenuStructureModule.h"
@@ -165,9 +169,22 @@
 #include "Filters/Widgets/Linking/SFilterBarSwitcher.h"
 #include "MVVM/ViewModels/OutlinerFilterAreaManager.h"
 #include "MVVM/ViewModels/SequencerEditorViewModelImpl.h"
+#include "SimpleView/SequencerSimpleViewSettings.h"
+#include "ToolableTimeline/Menus/ToolableTimelineMenuContext.h"
 
 #define LOCTEXT_NAMESPACE "Sequencer"
 
+namespace UE::Sequencer::Private
+{
+	static const TCHAR* const EnableSimpleViewCVarName = TEXT("Sequencer.SimpleView.Enabled");
+	static TAutoConsoleVariable<bool> CVarSequencerSimpleViewEnabled(
+		EnableSimpleViewCVarName,
+		false,
+		TEXT("When true, enables the Sequencer simple view for Sequencer. When false, Sequencer simple view is always disabled.\n")
+		TEXT("0: Disabled (default), 1: Enabled"),
+		ECVF_Default);
+}
+
 FSequencerOutlinerColumnVisibility::FSequencerOutlinerColumnVisibility(TSharedPtr<UE::Sequencer::IOutlinerColumn> InColumn)
 	: Column(InColumn)
 	, bIsColumnVisible(InColumn->IsColumnVisibleByDefault())
@@ -346,7 +363,7 @@ void SSequencer::Construct(const FArguments& InArgs, TSharedRef<FSequencer> InSe
 	TopTimeSlider = SequencerWidgets.CreateTimeSlider(TimeSliderControllerRef, bMirrorLabels);
 	bMirrorLabels = true;
 	BottomTimeSlider = SequencerWidgets.CreateTimeSlider(TimeSliderControllerRef, TAttribute<EVisibility>(this, &SSequencer::GetBottomTimeSliderVisibility), bMirrorLabels);
-	
+
 	// Create bottom time range slider
 	EShowRange Ranges = EShowRange(EShowRange::WorkingRange | EShowRange::ViewRange);
 	if(InArgs._ShowPlaybackRangeInTimeSlider)
@@ -354,15 +371,17 @@ void SSequencer::Construct(const FArguments& InArgs, TSharedRef<FSequencer> InSe
 		Ranges |= EShowRange::PlaybackRange;
 	}
 
+	const FTimeRangeArgs BottomTimeRangeArgs = FTimeRangeArgs(
+		Ranges,
+		TimeSliderControllerRef,
+		TAttribute<EVisibility>(this, &SSequencer::GetTimeRangeVisibility),
+		MakeAttributeSP(SequencerRef, &FSequencer::GetNumericTypeInterface, UE::Sequencer::ENumericIntent::Position),
+		GetSequencerSettings()->GetPlaybackRangeStartColor(),
+		GetSequencerSettings()->GetPlaybackRangeEndColor()
+	);
+
 	BottomTimeRange = SequencerWidgets.CreateTimeRange(
-		FTimeRangeArgs(
-			Ranges,
-			TimeSliderControllerRef,
-			TAttribute<EVisibility>(this, &SSequencer::GetTimeRangeVisibility),
-			MakeAttributeSP(SequencerRef, &FSequencer::GetNumericTypeInterface, ENumericIntent::Position),
-			SequencerSettings->GetPlaybackRangeStartColor(),
-			SequencerSettings->GetPlaybackRangeEndColor()
-		),
+		BottomTimeRangeArgs,
 		SequencerWidgets.CreateTimeRangeSlider(TimeSliderControllerRef)
 	);
 
@@ -389,6 +408,13 @@ void SSequencer::Construct(const FArguments& InArgs, TSharedRef<FSequencer> InSe
 		.FilterToggleButtonVisibility_Lambda(IsLinkedToggleButtonVisible)
 		.Sequencer(SequencerPtr);
 
+	if (SequencerRef->GetHostCapabilities().bSupportsSimpleView)
+	{
+		SimpleViewTimeline = StaticCastSharedRef<SimpleView::FSimpleViewTimeline>(
+			ToolableTimeline::FToolableTimeline::MakeInstance<SimpleView::FSimpleViewTimeline>(SequencerRef, TimeSliderArgs, BottomTimeRangeArgs)
+		);
+	}
+
 	ChildSlot
 	[
 		SNew(SVerticalBox)
@@ -397,162 +423,7 @@ void SSequencer::Construct(const FArguments& InArgs, TSharedRef<FSequencer> InSe
 		+ SVerticalBox::Slot()
 		.AutoHeight()
 		[
-			SNew(SBox)
-			.Visibility(this, &SSequencer::GetShowSequencerToolbar)
-			.Padding(FMargin(CommonPadding,0.0f,0.0f,0.f))
-			[
-				SNew(SWrapBox)
-				.UseAllottedSize(true)
-				.InnerSlotPadding(FVector2D(5, 0))
-				+ SWrapBox::Slot()
-				.FillEmptySpace(true)
-				.FillLineWhenSizeLessThan(600)
-				[
-					SAssignNew(ToolbarContainer, SBox)
-				]
-
-				+ SWrapBox::Slot()
-				.FillEmptySpace(true)
-				[
-					SNew(SHorizontalBox)
-					// Right Aligned Breadcrumbs
-					+ SHorizontalBox::Slot()
-					.HAlign(HAlign_Right)
-					.VAlign(VAlign_Center)
-					[
-						SNew(SSpacer)
-					]
-
-					// History Back Button
-					+ SHorizontalBox::Slot()
-					.HAlign(HAlign_Right)
-					.VAlign(VAlign_Center)
-					.AutoWidth()
-					[
-						SNew(SVerticalBox)
-
-						+ SVerticalBox::Slot()
-						.FillHeight(1.0f)
-						[
-							SNew(SButton)
-							.Visibility_Lambda([this] { return CanNavigateBreadcrumbs() ? EVisibility::Visible : EVisibility::Collapsed; } )
-							.VAlign(EVerticalAlignment::VAlign_Center)
-							.ButtonStyle(FAppStyle::Get(), "SimpleButton")
-							.ToolTipText_Lambda([this] { return SequencerPtr.Pin()->GetNavigateBackwardTooltip(); })
-							.ContentPadding(FMargin(1, 0))
-							.OnClicked_Lambda([this] { return SequencerPtr.Pin()->NavigateBackward(); })
-							.IsEnabled_Lambda([this] { return SequencerPtr.Pin()->CanNavigateBackward(); })
-							[
-								SNew(SBox) // scale up since the default icons are 16x16
-								.WidthOverride(20)
-								.HeightOverride(20)
-								[
-									SNew(SImage)
-									.ColorAndOpacity(FSlateColor::UseForeground())
-									.Image(FAppStyle::Get().GetBrush("Icons.ArrowLeft"))
-								]
-							]
-						]
-					]
-
-					// History Forward Button
-					+ SHorizontalBox::Slot()
-					.HAlign(HAlign_Right)
-					.VAlign(VAlign_Center)
-					.AutoWidth()
-					[
-						SNew(SVerticalBox)
-
-						+ SVerticalBox::Slot()
-						.FillHeight(1.0f)
-						[
-							SNew(SButton)
-							.Visibility_Lambda([this] { return CanNavigateBreadcrumbs() ? EVisibility::Visible : EVisibility::Collapsed; } )
-							.VAlign(EVerticalAlignment::VAlign_Center)
-							.ButtonStyle(FAppStyle::Get(), "SimpleButton")
-							.ToolTipText_Lambda([this] { return SequencerPtr.Pin()->GetNavigateForwardTooltip(); })
-							.ContentPadding(FMargin(1, 0))
-							.OnClicked_Lambda([this] { return SequencerPtr.Pin()->NavigateForward(); })
-							.IsEnabled_Lambda([this] { return SequencerPtr.Pin()->CanNavigateForward(); })
-							[
-								SNew(SBox) // scale up since the default icons are 16x16
-								.WidthOverride(20)
-								.HeightOverride(20)
-								[
-									SNew(SImage)
-									.ColorAndOpacity(FSlateColor::UseForeground())
-									.Image(FAppStyle::Get().GetBrush("Icons.ArrowRight"))
-								]
-							]
-						]
-					]
-
-					// Separator
-					+ SHorizontalBox::Slot()
-					.AutoWidth()
-					.Padding(3, 0)
-					[
-						SNew(SSeparator)
-						.Visibility_Lambda([this] { return CanNavigateBreadcrumbs() ? EVisibility::Visible : EVisibility::Collapsed; } )
-						.Orientation(Orient_Vertical)
-					]
-							
-					+ SHorizontalBox::Slot()
-					.HAlign(HAlign_Right)
-					.VAlign(VAlign_Center)
-					.AutoWidth()
-					[
-						SAssignNew(BreadcrumbPickerButton, SComboButton)
-						.Visibility_Lambda([this] { return CanNavigateBreadcrumbs() ? EVisibility::Visible : EVisibility::Collapsed; } )
-						.ButtonStyle(FAppStyle::Get(), "SimpleButton")
-						.OnGetMenuContent_Lambda([this] { return SNew(SSequencerHierarchyBrowser, SequencerPtr); })
-						.HasDownArrow(false)
-						.ContentPadding(FMargin(3, 3))
-						.ButtonContent()
-						[
-							SNew(SImage)
-							.ColorAndOpacity(FSlateColor::UseForeground())
-							.Image(FAppStyle::Get().GetBrush("Icons.FolderOpen"))
-						]
-					]
-
-					// Right Aligned Breadcrumbs
-					+ SHorizontalBox::Slot()
-					.HAlign(HAlign_Right)
-					.VAlign(VAlign_Center)
-					.AutoWidth()
-					[
-						SAssignNew(BreadcrumbTrail, SBreadcrumbTrail<FSequencerBreadcrumb>)
-						.Visibility(this, &SSequencer::GetBreadcrumbTrailVisibility)
-						.OnCrumbClicked(this, &SSequencer::OnCrumbClicked)
-						.ButtonStyle(FAppStyle::Get(), "SimpleButton")
-						.DelimiterImage(FAppStyle::Get().GetBrush("Sequencer.BreadcrumbIcon"))
-						.TextStyle(FAppStyle::Get(), "Sequencer.BreadcrumbText")
-					]
-
-					// Sequence Locking symbol
-					+ SHorizontalBox::Slot()
-					.HAlign(HAlign_Right)
-					.VAlign(VAlign_Center)
-					.AutoWidth()
-					[
-						SNew(SCheckBox)
-						.Style(FAppStyle::Get(),"ToggleButtonCheckBoxAlt")
-						.Type(ESlateCheckBoxType::CheckBox) // Use CheckBox instead of ToggleType since we're not putting ohter widget inside
-						.Padding(FMargin(0.f))
-						.IsFocusable(false)
-						.IsChecked_Lambda([this] { return GetIsSequenceReadOnly() ? ECheckBoxState::Checked : ECheckBoxState::Unchecked; } )
-						.OnCheckStateChanged(this, &SSequencer::OnSetSequenceReadOnly)
-						.ToolTipText_Lambda([this] { return GetIsSequenceReadOnly() ? LOCTEXT("UnlockSequence", "Unlock the animation so that it is editable") : LOCTEXT("LockSequence", "Lock the animation so that it is not editable"); } )
-						.CheckedImage(FAppStyle::Get().GetBrush("Icons.Lock"))
-						.CheckedHoveredImage(FAppStyle::Get().GetBrush("Icons.Lock"))
-						.CheckedPressedImage(FAppStyle::Get().GetBrush("Icons.Lock"))
-						.UncheckedImage(FAppStyle::Get().GetBrush("Icons.Unlock"))
-						.UncheckedHoveredImage(FAppStyle::Get().GetBrush("Icons.Unlock"))
-						.UncheckedPressedImage(FAppStyle::Get().GetBrush("Icons.Unlock"))
-					]
-				]
-			]
+			ConstructToolBar()
 		]
 
 		// Main content body
@@ -563,7 +434,7 @@ void SSequencer::Construct(const FArguments& InArgs, TSharedRef<FSequencer> InSe
 		]
 	];
 
-	RebuildForSidebar();
+	RebuildMainContent();
 
 	if (InSequencer->GetHostCapabilities().bSupportsCurveEditor)
 	{
@@ -577,11 +448,284 @@ void SSequencer::Construct(const FArguments& InArgs, TSharedRef<FSequencer> InSe
 	InSequencer->GetViewModel()->GetSelection()->Outliner.OnChanged.AddSP(this, &SSequencer::HandleOutlinerNodeSelectionChanged);
 
 	ResetBreadcrumbs();
+
+	// Make sure the saved state (full or simplified view) is restored correctly.
+	// This needs to be done after the geometry is constructed, so we do it next tick.
+	if (GEditor)
+	{
+		const TWeakPtr<SSequencer> ThisWeak = SharedThis(this);
+		GEditor->GetTimerManager()->SetTimerForNextTick([ThisWeak]()
+			{
+				if (const TSharedPtr<SSequencer> This = ThisWeak.Pin())
+				{
+					This->EnableSimpleView(This->IsSimpleView());
+				}
+			});
+	}
+
+	if (IConsoleVariable* const EnableSimpleViewCVar = IConsoleManager::Get().FindConsoleVariable(Private::EnableSimpleViewCVarName))
+	{
+		EnableSimpleViewCVarDelegateHandle = EnableSimpleViewCVar->OnChangedDelegate().AddLambda([this](IConsoleVariable* const InConsoleVariable)
+			{
+				// Return to full Sequencer view when the cvar is disabled
+				if (!Private::CVarSequencerSimpleViewEnabled.GetValueOnAnyThread() && IsSimpleView())
+				{
+					EnableSimpleView(false);
+				}
+			});
+	}
 }
 
-void SSequencer::RebuildForSidebar()
+void SSequencer::RebuildMainContent()
 {
-	MainContentContainer->SetContent(ConstructSidebarContent());
+	const TSharedPtr<FSequencer> Sequencer = SequencerPtr.Pin();
+	check(Sequencer.IsValid());
+
+	TSharedPtr<SWidget> MainContent;
+
+	if (Sequencer->GetHostCapabilities().bSupportsSimpleView)
+	{
+		MainContent = SNew(SWidgetSwitcher)
+			.WidgetIndex_Lambda([this]()
+				{
+					return GetSequencerSettings()->IsSimpleView() ? 1 : 0;
+				})
+			+ SWidgetSwitcher::Slot()
+			.VAlign(VAlign_Fill)
+			[
+				ConstructSidebarContent()
+			]
+			+ SWidgetSwitcher::Slot()
+			.VAlign(VAlign_Top)
+			[
+				SNew(SBox)
+				.HAlign(HAlign_Fill)
+				.VAlign(VAlign_Top)
+				[
+					SimpleViewTimeline->GenerateWidget()
+				]
+			];
+	}
+	else
+	{
+		MainContent = ConstructSidebarContent();
+	}
+
+	MainContentContainer->SetContent(MainContent.ToSharedRef());
+}
+
+TSharedRef<SWidget> SSequencer::ConstructToolBar()
+{
+	const TSharedPtr<FSequencer> Sequencer = SequencerPtr.Pin();
+
+	const TSharedRef<SHorizontalBox> Content = SNew(SHorizontalBox)
+		// Right Aligned Breadcrumbs
+		+ SHorizontalBox::Slot()
+		.HAlign(HAlign_Right)
+		.VAlign(VAlign_Center)
+		[
+			SNew(SSpacer)
+		]
+		// History Back Button
+		+ SHorizontalBox::Slot()
+		.HAlign(HAlign_Right)
+		.VAlign(VAlign_Center)
+		.AutoWidth()
+		[
+			SNew(SVerticalBox)
+
+			+ SVerticalBox::Slot()
+			.FillHeight(1.0f)
+			[
+				SNew(SButton)
+				.Visibility(this, &SSequencer::GetBreadcrumbNavigateVisibility)
+				.VAlign(EVerticalAlignment::VAlign_Center)
+				.ButtonStyle(FAppStyle::Get(), "SimpleButton")
+				.ToolTipText_Lambda([this] { return SequencerPtr.Pin()->GetNavigateBackwardTooltip(); })
+				.ContentPadding(FMargin(1, 0))
+				.OnClicked_Lambda([this] { return SequencerPtr.Pin()->NavigateBackward(); })
+				.IsEnabled_Lambda([this] { return SequencerPtr.Pin()->CanNavigateBackward(); })
+				[
+					SNew(SBox) // scale up since the default icons are 16x16
+					.WidthOverride(20)
+					.HeightOverride(20)
+					[
+						SNew(SImage)
+						.ColorAndOpacity(FSlateColor::UseForeground())
+						.Image(FAppStyle::Get().GetBrush("Icons.ArrowLeft"))
+					]
+				]
+			]
+		]
+		// History Forward Button
+		+ SHorizontalBox::Slot()
+		.HAlign(HAlign_Right)
+		.VAlign(VAlign_Center)
+		.AutoWidth()
+		[
+			SNew(SVerticalBox)
+
+			+ SVerticalBox::Slot()
+			.FillHeight(1.0f)
+			[
+				SNew(SButton)
+				.Visibility(this, &SSequencer::GetBreadcrumbNavigateVisibility)
+				.VAlign(EVerticalAlignment::VAlign_Center)
+				.ButtonStyle(FAppStyle::Get(), "SimpleButton")
+				.ToolTipText_Lambda([this] { return SequencerPtr.Pin()->GetNavigateForwardTooltip(); })
+				.ContentPadding(FMargin(1, 0))
+				.OnClicked_Lambda([this] { return SequencerPtr.Pin()->NavigateForward(); })
+				.IsEnabled_Lambda([this] { return SequencerPtr.Pin()->CanNavigateForward(); })
+				[
+					SNew(SBox) // scale up since the default icons are 16x16
+					.WidthOverride(20)
+					.HeightOverride(20)
+					[
+						SNew(SImage)
+						.ColorAndOpacity(FSlateColor::UseForeground())
+						.Image(FAppStyle::Get().GetBrush("Icons.ArrowRight"))
+					]
+				]
+			]
+		]
+		+ SHorizontalBox::Slot()
+		.AutoWidth()
+		.HAlign(HAlign_Right)
+		.VAlign(VAlign_Center)
+		[
+			SNew(SHorizontalBox)
+			// Separator
+			+ SHorizontalBox::Slot()
+			.AutoWidth()
+			.Padding(3, 0)
+			[
+				SNew(SSeparator)
+				.Visibility(this, &SSequencer::GetBreadcrumbNavigateVisibility)
+				.Orientation(Orient_Vertical)
+			]
+			// Breadcrumb Picker
+			+ SHorizontalBox::Slot()
+			.HAlign(HAlign_Right)
+			.VAlign(VAlign_Center)
+			.AutoWidth()
+			[
+				SAssignNew(BreadcrumbPickerButton, SComboButton)
+				.Visibility(this, &SSequencer::GetBreadcrumbNavigateVisibility)
+				.ButtonStyle(FAppStyle::Get(), "SimpleButton")
+				.OnGetMenuContent_Lambda([this]
+					{
+						return SNew(SSequencerHierarchyBrowser, SequencerPtr);
+					})
+				.HasDownArrow(false)
+				.ContentPadding(FMargin(3, 3))
+				.ButtonContent()
+				[
+					SNew(SImage)
+					.ColorAndOpacity(FSlateColor::UseForeground())
+					.Image(FAppStyle::Get().GetBrush("Icons.FolderOpen"))
+				]
+			]
+			// Right Aligned Breadcrumbs
+			+ SHorizontalBox::Slot()
+			.HAlign(HAlign_Right)
+			.VAlign(VAlign_Center)
+			.AutoWidth()
+			[
+				SAssignNew(BreadcrumbTrail, SBreadcrumbTrail<FSequencerBreadcrumb>)
+				.Visibility(this, &SSequencer::GetBreadcrumbTrailVisibility)
+				.OnCrumbClicked(this, &SSequencer::OnCrumbClicked)
+				.ButtonStyle(FAppStyle::Get(), "SimpleButton")
+				.DelimiterImage(FAppStyle::Get().GetBrush("Sequencer.BreadcrumbIcon"))
+				.TextStyle(FAppStyle::Get(), "Sequencer.BreadcrumbText")
+			]
+			// Sequence Locking symbol
+			+ SHorizontalBox::Slot()
+			.HAlign(HAlign_Right)
+			.VAlign(VAlign_Center)
+			.AutoWidth()
+			.Padding(0, 0, 6, 0)
+			[
+				SNew(SCheckBox)
+				.Style(FAppStyle::Get(),"ToggleButtonCheckBoxAlt")
+				.ForegroundColor_Lambda([this]()
+					{
+						return IsReadOnly() ? FStyleColors::AccentRed : FStyleColors::Foreground;
+					})
+				.Type(ESlateCheckBoxType::CheckBox) // Use CheckBox instead of ToggleType since we're not putting ohter widget inside
+				.Padding(FMargin(0.f))
+				.IsFocusable(false)
+				.IsChecked_Lambda([this]
+					{
+						return IsReadOnly() ? ECheckBoxState::Checked : ECheckBoxState::Unchecked;
+					})
+				.OnCheckStateChanged_Lambda([this](const ECheckBoxState InNewState)
+					{
+						SetReadOnly(InNewState == ECheckBoxState::Checked);
+					})
+				.ToolTipText_Lambda([this]
+					{
+						return IsReadOnly()
+							? LOCTEXT("UnlockSequence", "Unlock the animation so that it is editable")
+							: LOCTEXT("LockSequence", "Lock the animation so that it is not editable");
+					})
+				.CheckedImage(FAppStyle::Get().GetBrush("Icons.Lock"))
+				.CheckedHoveredImage(FAppStyle::Get().GetBrush("Icons.Lock"))
+				.CheckedPressedImage(FAppStyle::Get().GetBrush("Icons.Lock"))
+				.UncheckedImage(FAppStyle::Get().GetBrush("Icons.Unlock"))
+				.UncheckedHoveredImage(FAppStyle::Get().GetBrush("Icons.Unlock"))
+				.UncheckedPressedImage(FAppStyle::Get().GetBrush("Icons.Unlock"))
+			]
+		];
+
+	// Toggle button to switch between full and simplified views
+	if (SimpleViewTimeline.IsValid())
+	{
+		Content->AddSlot()
+			.AutoWidth()
+			[
+				SNew(SHorizontalBox)
+				.Visibility_Lambda([]()
+					{
+						return UE::Sequencer::Private::CVarSequencerSimpleViewEnabled.GetValueOnAnyThread()
+							? EVisibility::Visible : EVisibility::Collapsed;
+					})
+				+ SHorizontalBox::Slot()
+				.AutoWidth()
+				.Padding(0, 0, 3, 0)
+				[
+					SNew(SSeparator)
+					.Orientation(Orient_Vertical)
+				]
+				+ SHorizontalBox::Slot()
+				.AutoWidth()
+				.HAlign(HAlign_Right)
+				.VAlign(VAlign_Center)
+				.Padding(3, 0, 6, 0)
+				[
+					UE::Sequencer::SimpleView::Utils::CreateToggleButton(SimpleViewTimeline)
+				]
+			];
+	}
+
+	return SNew(SBox)
+		.Visibility(this, &SSequencer::GetShowSequencerToolbar)
+		.Padding(FMargin(CommonPadding,0.0f,0.0f,0.f))
+		[
+			SNew(SWrapBox)
+			.UseAllottedSize(true)
+			.InnerSlotPadding(FVector2D(5, 0))
+			+ SWrapBox::Slot()
+			.FillEmptySpace(true)
+			.FillLineWhenSizeLessThan(600)
+			[
+				SAssignNew(ToolbarContainer, SBox)
+			]
+			+ SWrapBox::Slot()
+			.FillEmptySpace(true)
+			[
+				Content
+			]
+		];
 }
 
 TSharedRef<SWidget> SSequencer::ConstructSidebarContent()
@@ -1199,6 +1343,16 @@ void SSequencer::BindCommands(TSharedRef<FUICommandList> SequencerCommandBinding
 		FSequencerCommands::Get().OpenNodeGroupsManager,
 		FExecuteAction::CreateSP(this, &SSequencer::OpenNodeGroupsManager)
 	);
+
+	SequencerCommandBindings->MapAction(
+		FSequencerCommands::Get().ToggleToolBarVisible,
+		FExecuteAction::CreateSP(this, &SSequencer::ToggleToolBarVisible)
+	);
+
+	SequencerCommandBindings->MapAction(
+		FSequencerCommands::Get().ToggleSimpleView,
+		FExecuteAction::CreateSP(this, &SSequencer::ToggleSimpleView)
+	);
 }
 
 void SSequencer::OpenTickResolutionOptions()
@@ -1231,7 +1385,6 @@ void SSequencer::OpenTickResolutionOptions()
 	WeakTickResolutionOptionsWindow = TickResolutionOptionsWindow;
 }
 
-
 void SSequencer::NotifyPostChange(const FPropertyChangedEvent& PropertyChangedEvent, FEditPropertyChain* PropertyThatChanged)
 {
 	// @todo sequencer: is this still needed?
@@ -1373,12 +1526,13 @@ const FLazyName ViewOptionsMenuName("Sequencer.MainToolBar.ViewOptions");
 TSharedRef<SWidget> SSequencer::MakeToolBar()
 {
 	using namespace UE::Sequencer::ToolMenus;
-	
+
 	if (!UToolMenus::Get()->IsMenuRegistered(ToolbarMenuName))
 	{
 		UToolMenu* Toolbar = UToolMenus::Get()->RegisterMenu(ToolbarMenuName, NAME_None, EMultiBoxType::ToolBar);
 		Toolbar->AddDynamicSection("PopulateToolBar", FNewToolMenuDelegate::CreateStatic(&SSequencer::PopulateToolBar));
 	}
+
 	RegisterViewOptionsToolMenu();
 
 	TArray<TSharedPtr<FExtender>> AllExtenders;
@@ -1395,6 +1549,13 @@ TSharedRef<SWidget> SSequencer::MakeToolBar()
 
 	FToolMenuContext Context(SequencerPtr.Pin()->GetCommandBindings(), Extender, ContextObject);
 
+	if (Sequencer.IsValid() && Sequencer->GetHostCapabilities().bSupportsSimpleView)
+	{
+		UToolableTimelineMenuContext* const SimpleViewContext = NewObject<UToolableTimelineMenuContext>();
+		SimpleViewContext->WeakTimeline = SimpleViewTimeline;
+		Context.AddObject(SimpleViewContext);
+	}
+
 	// Allow any toolkits to initialize their menu context
 	OnInitToolMenuContext.ExecuteIfBound(Context);
 
@@ -1420,6 +1581,12 @@ void SSequencer::PopulateToolBar(UToolMenu* InMenu)
 
 	const TSharedRef<SSequencer> SequencerWidget = Sequencer->GetUnderlyingSequencerWidget();
 
+	const TAttribute<EVisibility> FullViewVisibility = TAttribute<EVisibility>::CreateLambda([WeakSequencer]()
+		{
+			const TSharedPtr<FSequencer> Sequencer = WeakSequencer.Pin();
+			return (Sequencer.IsValid() && !Sequencer->IsSimpleView()) ? EVisibility::Visible : EVisibility::Collapsed;
+		});
+
 	{
 		FToolMenuSection& Section = InMenu->AddSection("BaseCommands");
 
@@ -1459,16 +1626,19 @@ void SSequencer::PopulateToolBar(UToolMenu* InMenu)
 
 			FToolMenuEntry FindInContentBrowserEntry = FToolMenuEntry::InitToolBarButton(FSequencerCommands::Get().FindInContentBrowser, TAttribute<FText>(), TAttribute<FText>(), FSlateIcon(FAppStyle::Get().GetStyleSetName(), "SystemWideCommands.FindInContentBrowser"));
 			FindInContentBrowserEntry.StyleNameOverride = SequencerToolbarStyleName;
+			FindInContentBrowserEntry.Visibility = FullViewVisibility;
 			Section.AddEntry(FindInContentBrowserEntry);
 
 			FToolMenuEntry CreateCameraEntry = FToolMenuEntry::InitToolBarButton(FSequencerCommands::Get().CreateCamera);
 			CreateCameraEntry.StyleNameOverride = SequencerToolbarStyleName;
+			CreateCameraEntry.Visibility = FullViewVisibility;
 			Section.AddEntry(CreateCameraEntry);
 
 			if (Sequencer->GetHostCapabilities().bSupportsRenderMovie)
 			{
 				FToolMenuEntry RenderMovieEntry = FToolMenuEntry::InitToolBarButton(FSequencerCommands::Get().RenderMovie, TAttribute<FText>(), TAttribute<FText>(), FSlateIcon(FAppStyle::GetAppStyleSetName(), "LevelEditor.OpenCinematic"));
 				RenderMovieEntry.StyleNameOverride = SequencerToolbarStyleName;
+				RenderMovieEntry.Visibility = FullViewVisibility;
 				Section.AddEntry(RenderMovieEntry);
 			
 				FToolMenuEntry RenderMovieOptionsEntry = FToolMenuEntry::InitComboButton(
@@ -1481,6 +1651,7 @@ void SSequencer::PopulateToolBar(UToolMenu* InMenu)
 					true
 				);
 				RenderMovieOptionsEntry.StyleNameOverride = SequencerToolbarStyleName;
+				RenderMovieOptionsEntry.Visibility = FullViewVisibility;
 				Section.AddEntry(RenderMovieOptionsEntry);
 			}
 
@@ -1490,6 +1661,7 @@ void SSequencer::PopulateToolBar(UToolMenu* InMenu)
 				// Only show this button where it makes sense (ie, if the sequence is not contained within a blueprint already)
 				FToolMenuEntry OpenDirectorBlueprintEntry = FToolMenuEntry::InitToolBarButton(FSequencerCommands::Get().OpenDirectorBlueprint, TAttribute<FText>(), TAttribute<FText>(), FSlateIcon(FAppStyle::Get().GetStyleSetName(), "LevelEditor.OpenLevelBlueprint"));
 				OpenDirectorBlueprintEntry.StyleNameOverride = SequencerToolbarStyleName;
+				OpenDirectorBlueprintEntry.Visibility = FullViewVisibility;
 				Section.AddEntry(OpenDirectorBlueprintEntry);
 			}
 
@@ -1505,6 +1677,7 @@ void SSequencer::PopulateToolBar(UToolMenu* InMenu)
 			FSlateIcon(FAppStyle::GetAppStyleSetName(), "Sequencer.Actions")
 		);
 		ActionsEntry.StyleNameOverride = SequencerToolbarStyleName;
+		ActionsEntry.Visibility = FullViewVisibility;
 		Section.AddEntry(ActionsEntry);
 
 		FToolMenuEntry ViewOptionsEntry = FToolMenuEntry::InitComboButton(
@@ -2527,6 +2700,14 @@ SSequencer::~SSequencer()
 {
 	using namespace UE::Sequencer;
 
+	if (EnableSimpleViewCVarDelegateHandle.IsValid())
+	{
+		if (IConsoleVariable* const EnableSimpleViewCVar = IConsoleManager::Get().FindConsoleVariable(Private::EnableSimpleViewCVarName))
+		{
+			EnableSimpleViewCVar->OnChangedDelegate().Remove(EnableSimpleViewCVarDelegateHandle);
+		}
+	}
+
 	USelection::SelectionChangedEvent.RemoveAll(this);
 
 	TSharedPtr<FSequencer> Sequencer = SequencerPtr.Pin();
@@ -2804,7 +2985,6 @@ void SSequencer::UpdateBreadcrumbs()
 	}
 }
 
-
 void SSequencer::ResetBreadcrumbs()
 {
 	BreadcrumbTrail->ClearCrumbs();
@@ -2946,17 +3126,24 @@ FReply SSequencer::OnDrop( const FGeometry& MyGeometry, const FDragDropEvent& Dr
 	return bWasDropHandled ? FReply::Handled() : FReply::Unhandled();
 }
 
-
 FReply SSequencer::OnKeyDown( const FGeometry& MyGeometry, const FKeyEvent& InKeyEvent ) 
 {
 	// A toolkit tab is active, so direct all command processing to it
-	const TSharedPtr<FSequencer> SequencerPin = SequencerPtr.Pin();
-	if (!SequencerPin)
+	const TSharedPtr<FSequencer> Sequencer = SequencerPtr.Pin();
+	if (!Sequencer.IsValid())
 	{
 		return FReply::Unhandled();
 	}
-	
-	if (SequencerPin->GetCommandBindings()->ProcessCommandBindings(InKeyEvent))
+
+	// Give the independent timeline commands a chance to process before Sequencer
+	if (SimpleViewTimeline.IsValid()
+		&& SimpleViewTimeline->IsInSimpleView()
+		&& SimpleViewTimeline->GetCommandList()->ProcessCommandBindings(InKeyEvent))
+	{
+		return FReply::Handled();
+	}
+
+	if (Sequencer->GetCommandBindings()->ProcessCommandBindings(InKeyEvent))
 	{
 		return FReply::Handled();
 	}
@@ -3357,6 +3544,10 @@ FText SSequencer::GetBreadcrumbTextForSequence(TWeakObjectPtr<UMovieSceneSequenc
 	}
 }
 
+EVisibility SSequencer::GetBreadcrumbNavigateVisibility() const
+{
+	return CanNavigateBreadcrumbs() ? EVisibility::Visible : EVisibility::Collapsed;
+}
 
 EVisibility SSequencer::GetBreadcrumbTrailVisibility() const
 {
@@ -3408,7 +3599,15 @@ EVisibility SSequencer::GetShowTickLines() const
 
 EVisibility SSequencer::GetShowSequencerToolbar() const
 {
-	return GetSequencerSettings()->GetShowSequencerToolbar() ? EVisibility::Visible : EVisibility::Collapsed;
+	if (SimpleViewTimeline.IsValid() && IsSimpleView())
+	{
+		const USequencerSimpleViewSettings* const SimpleViewSettings = GetDefault<USequencerSimpleViewSettings>();
+		const bool bShowToolbar = SimpleViewSettings->bShowToolBar;
+		return bShowToolbar ? EVisibility::SelfHitTestInvisible : EVisibility::Collapsed;
+	}
+
+	return GetSequencerSettings()->GetShowSequencerToolbar()
+		? EVisibility::SelfHitTestInvisible : EVisibility::Collapsed;
 }
 
 EFrameNumberDisplayFormats SSequencer::GetTimeDisplayFormat() const
@@ -3697,35 +3896,36 @@ double SSequencer::GetSpinboxDelta() const
 	return Sequencer->GetDisplayRateDeltaFrameCount();
 }
 
-bool SSequencer::GetIsSequenceReadOnly() const
+bool SSequencer::IsReadOnly() const
 {
-	TSharedPtr<FSequencer> Sequencer = SequencerPtr.Pin();
-	return Sequencer->GetFocusedMovieSceneSequence() ? Sequencer->GetFocusedMovieSceneSequence()->GetMovieScene()->IsReadOnly() : false;
+	const TSharedPtr<FSequencer> Sequencer = SequencerPtr.Pin();
+	return Sequencer.IsValid() ? Sequencer->IsReadOnly() : false;
 }
 
-void SSequencer::OnSetSequenceReadOnly(ECheckBoxState CheckBoxState)
+void SSequencer::SetReadOnly(const bool bInReadOnly)
 {
 	TSharedPtr<FSequencer> Sequencer = SequencerPtr.Pin();
-	
-	bool bReadOnly = CheckBoxState == ECheckBoxState::Checked;
 
 	if (Sequencer->GetFocusedMovieSceneSequence())
 	{
-		UMovieScene* MovieScene = Sequencer->GetFocusedMovieSceneSequence()->GetMovieScene();
-		const FScopedTransaction Transaction(CheckBoxState == ECheckBoxState::Checked ? LOCTEXT("LockMovieScene", "Lock Movie Scene") : LOCTEXT("UnlockMovieScene", "Unlock Movie Scene") );
+		const FText TransactionText = bInReadOnly
+			? LOCTEXT("LockMovieScene", "Lock Movie Scene")
+			: LOCTEXT("UnlockMovieScene", "Unlock Movie Scene");
+		const FScopedTransaction Transaction(TransactionText);
 
+		UMovieScene* MovieScene = Sequencer->GetFocusedMovieSceneSequence()->GetMovieScene();
 		MovieScene->Modify();
-		MovieScene->SetReadOnly(bReadOnly);
+		MovieScene->SetReadOnly(bInReadOnly);
 
 		TArray<UMovieScene*> DescendantMovieScenes;
 		MovieSceneHelpers::GetDescendantMovieScenes(Sequencer->GetFocusedMovieSceneSequence(), DescendantMovieScenes);
 
 		for (UMovieScene* DescendantMovieScene : DescendantMovieScenes)
 		{
-			if (DescendantMovieScene && bReadOnly != DescendantMovieScene->IsReadOnly())
+			if (DescendantMovieScene && bInReadOnly != DescendantMovieScene->IsReadOnly())
 			{
 				DescendantMovieScene->Modify();
-				DescendantMovieScene->SetReadOnly(bReadOnly);
+				DescendantMovieScene->SetReadOnly(bInReadOnly);
 			}
 		}
 
@@ -3942,7 +4142,7 @@ void SSequencer::SetSidebarVisible(const bool bInVisible)
 	SidebarState.SetVisible(bInVisible);
 	SequencerSettings->SetSidebarState(SidebarState);
 
-	RebuildForSidebar();
+	RebuildMainContent();
 }
 
 void SSequencer::ToggleSidebarVisible()
@@ -4302,12 +4502,6 @@ EVisibility SSequencer::GetPopoutTransportControlsVisibility() const
 		return EVisibility::Collapsed;
 	}
 
-	const UMovieSceneSequence* const RootSequence = Sequencer->GetRootMovieSceneSequence();
-	if (!RootSequence || !RootSequence->IsA<ULevelSequence>())
-	{
-		return EVisibility::Visible;
-	}
-
 	const USequencerSettings* const SequencerSettings = Sequencer->GetSequencerSettings();
 	if (!SequencerSettings)
 	{
@@ -4315,8 +4509,41 @@ EVisibility SSequencer::GetPopoutTransportControlsVisibility() const
 	}
 
 	const bool ShowTransportControls = SequencerSettings->ShouldShowTransportControls();
-	const bool TransportControlsPoppedOut = SequencerSettings->AreTransportControlsPoppedOut();
-	return (!ShowTransportControls || TransportControlsPoppedOut) ? EVisibility::Collapsed : EVisibility::Visible;
+
+	return UE::Sequencer::IsPopoutTransportControlsVisible(Sequencer.ToSharedRef(), ShowTransportControls)
+		? EVisibility::Visible : EVisibility::Collapsed;
+}
+
+bool SSequencer::IsSimpleView() const
+{
+	return SimpleViewTimeline.IsValid() ? SimpleViewTimeline->IsInSimpleView() : false;
+}
+
+void SSequencer::EnableSimpleView(const bool bInEnable)
+{
+	if (SimpleViewTimeline.IsValid())
+	{
+		const bool bEnableSimpleView = UE::Sequencer::Private::CVarSequencerSimpleViewEnabled.GetValueOnAnyThread()
+			? bInEnable : false;
+		SimpleViewTimeline->EnableSimpleView(bEnableSimpleView);
+
+		RebuildMainContent();
+	}
+}
+
+void SSequencer::ToggleSimpleView()
+{
+	EnableSimpleView(!IsSimpleView());
+}
+
+void SSequencer::ToggleToolBarVisible()
+{
+	const TSharedPtr<FSequencer> Sequencer = SequencerPtr.Pin();
+	USequencerSettings* const SequencerSettings = Sequencer.IsValid() ? Sequencer->GetSequencerSettings() : nullptr;
+	if (SequencerSettings)
+	{
+		SequencerSettings->SetShowSequencerToolbar(!SequencerSettings->GetShowSequencerToolbar());
+	}
 }
 
 #undef LOCTEXT_NAMESPACE
diff --git a/Engine/Source/Editor/Sequencer/Private/SSequencer.h b/Engine/Source/Editor/Sequencer/Private/SSequencer.h
--- a/Engine/Source/Editor/Sequencer/Private/SSequencer.h
+++ b/Engine/Source/Editor/Sequencer/Private/SSequencer.h
@@ -63,13 +63,17 @@ namespace UE
 namespace Sequencer
 {
 class SFilterBarSwitcher;
+class SOutlinerView;
 class SSequencerSearchAndFilterRow;
+class STrackAreaView;
+class FVirtualTrackArea;
+class IOutlinerColumn;
+struct FSequencerSelectionCurveFilter;
 
-class SOutlinerView;
-	class STrackAreaView;
-	class FVirtualTrackArea;
-	class IOutlinerColumn;
-	struct FSequencerSelectionCurveFilter;
+namespace SimpleView
+{
+	class FSimpleViewTimeline;
+}
 
 } // namespace Sequencer
 } // namespace UE
@@ -380,6 +384,7 @@ private:
 
 	static constexpr float CommonPadding = 3.f;
 
+	TSharedRef<SWidget> ConstructToolBar();
 	TSharedRef<SWidget> ConstructSidebarContent();
 	TSharedRef<SWidget> ConstructFilterBarContent();
 	TSharedRef<SSplitter> ConstructOutlinerAndTrackArea(const FArguments& InArgs, const FTimeSliderArgs& InTimeSliderArgs);
@@ -387,8 +392,11 @@ private:
 	TSharedRef<SWidget> ConstructTransportPanel(const TSharedRef<ISequencer>& InSequencer);
 	TSharedRef<SWidget> ConstructTransportControls(const TSharedRef<ISequencer>& InSequencer);
 
-	/** Can be called anytime after Construct() to rebuild most of the main content (except for the toolbar). */
-	void RebuildForSidebar();
+	/** Can be called anytime after Construct() to rebuild most of the main content (except for the toolbar).
+	 * If the sidebar visibility, sidebar layout orientation, or host capabilities for simple view change, this function must be called. */
+	void RebuildMainContent();
+
+	/** Can be called anytime after Construct() to rebuild most of the main content */
 	void RebuildFilterBarContent();
 	void RebuildSearchAndFilterRow();
 
@@ -471,6 +479,9 @@ public:
 	/** Gets the top time sliders widget. */
 	TSharedPtr<ITimeSlider> GetTopTimeSliderWidget() const;
 
+	bool IsReadOnly() const;
+	void SetReadOnly(const bool bInReadOnly);
+
 private:
 
 	TRange<FFrameNumber> GetViewSpacePlaybackRange(TAttribute<TRange<FFrameNumber>> Range) const;
@@ -547,12 +558,17 @@ private:
 	FText GetBreadcrumbTextForSection(TWeakObjectPtr<UMovieSceneSubSection> SubSection) const;
 	FText GetBreadcrumbTextForSequence(TWeakObjectPtr<UMovieSceneSequence> Sequence, bool bIsActive) const;
 
+	/** @return Visibility of the breadcrumb navigation widgets */
+	EVisibility GetBreadcrumbNavigateVisibility() const;
+
 	/** Gets whether or not the breadcrumb trail should be visible. */
 	EVisibility GetBreadcrumbTrailVisibility() const;
 
+public:
 	/** Return whether there are breadcrumbs to navigate. */
 	bool CanNavigateBreadcrumbs() const;
 
+private:
 	/** Gets whether or not the bottom time slider should be visible. */
 	EVisibility GetBottomTimeSliderVisibility() const;
 
@@ -575,18 +591,12 @@ private:
 	FPaintPlaybackRangeArgs GetSectionPlaybackRangeArgs() const;
 
 	EVisibility GetDebugVisualizerVisibility() const;
-	
-	void SetPlaybackSpeed(float InPlaybackSpeed);
-	float GetPlaybackSpeed() const;
 
 	/** Controls how fast Spinboxes change values. */
 	double GetSpinboxDelta() const;
 
-	bool GetIsSequenceReadOnly() const;
-	void OnSetSequenceReadOnly(ECheckBoxState CheckBoxState);
-
 	/** Returns whether or not the Curve Editor is enabled. Allows us to bind to the Slate Enabled attribute. */
-	bool GetIsCurveEditorEnabled() const { return !GetIsSequenceReadOnly(); }
+	bool GetIsCurveEditorEnabled() const { return !IsReadOnly(); }
 
 public:
 	/** On Paste Command */
@@ -691,6 +701,14 @@ public:
 	void SetTransportControlsPoppedOut(const bool bInPoppedOut);
 	void ToggleTransportControlsPoppedOut();
 
+	TSharedPtr<FTabManager> GetTabManager() const;
+
+	bool IsSimpleView() const;
+	void EnableSimpleView(const bool bInEnable);
+	void ToggleSimpleView();
+
+	void ToggleToolBarVisible();
+
 private:
 
 	/** Pending focus handler */
@@ -704,8 +722,6 @@ private:
 
 	void OnSidebarStateChanged(const FSidebarState& InNewState);
 
-	TSharedPtr<FTabManager> GetTabManager() const;
-
 	void RegisterSelectionDetailsTab();
 	void UnregisterSelectionDetailsTab();
 
@@ -788,7 +804,7 @@ private:
 	SVerticalBox::FSlot* OutlinerPinnedSlot = nullptr;
 
 	/** Width of the outliner tree. Mainly used while sizing so we don't need to SaveConfig() and query settings, which causes a performance hit */
-	float TreeViewWidth = 0.f;
+	mutable float TreeViewWidth = 0.f;
 
 	/** List of registered outliner columns with their visibility states */
 	TArray<FSequencerOutlinerColumnVisibility> OutlinerColumnVisibilities;
@@ -895,4 +911,10 @@ private:
 
 	/** Widget that manages popped out state (tab window registration) of the transport controls */
 	TSharedPtr<SPopoutTabInlineContent> TransportControlsPopoutTab;
+
+	/** Toolable timeline used for the Sequencer simple view */
+	TSharedPtr<UE::Sequencer::SimpleView::FSimpleViewTimeline> SimpleViewTimeline;
+
+	/** Handle to callback to toggle simple view off when CVarSequencerSimpleViewEnabled is disabled */
+	FDelegateHandle EnableSimpleViewCVarDelegateHandle;
 };
diff --git a/Engine/Source/Editor/Sequencer/Private/Sequencer.cpp b/Engine/Source/Editor/Sequencer/Private/Sequencer.cpp
--- a/Engine/Source/Editor/Sequencer/Private/Sequencer.cpp
+++ b/Engine/Source/Editor/Sequencer/Private/Sequencer.cpp
@@ -819,6 +819,8 @@ void FSequencer::InitSequencer(const FSequencerInitParams& InitParams
 
 		SequencerWidget->SetSelectionDetailsPoppedOut(Settings->IsSelectionDetailsPoppedOut());
 	}
+
+	SynchronizeSequencerSelectionWithExternalSelection();
 }
 
 void FSequencer::SetSequencerSettings(USequencerSettings* InSettings)
@@ -13132,6 +13134,19 @@ bool FSequencer::UnregisterDrawerSection(const FName InDrawerId, const FName InS
 	return false;
 }
 
+bool FSequencer::IsSimpleView() const
+{
+	return SequencerWidget.IsValid() ? SequencerWidget->IsSimpleView() : false;
+}
+
+void FSequencer::EnableSimpleView(const bool bInEnable)
+{
+	if (SequencerWidget.IsValid())
+	{
+		return SequencerWidget->EnableSimpleView(bInEnable);
+	}
+}
+
 void FSequencer::ToggleSidebar()
 {
 	if (SequencerWidget.IsValid())
diff --git a/Engine/Source/Editor/Sequencer/Private/Sequencer.h b/Engine/Source/Editor/Sequencer/Private/Sequencer.h
--- a/Engine/Source/Editor/Sequencer/Private/Sequencer.h
+++ b/Engine/Source/Editor/Sequencer/Private/Sequencer.h
@@ -1279,6 +1279,12 @@ public:
 	 */
 	virtual bool UnregisterDrawerSection(const FName InDrawerId, const FName InSectionId) override;
 
+	/** @return True if the sequencer is in the simplified view */
+	virtual bool IsSimpleView() const override;
+
+	/** Enables or disables the Sequencer simplified view */
+	virtual void EnableSimpleView(const bool bInEnable) override;
+
 	/** Toggles the visibility of the Sequencer sidebar. */
 	void ToggleSidebar();
 
diff --git a/Engine/Source/Editor/Sequencer/Private/SequencerCommands.cpp b/Engine/Source/Editor/Sequencer/Private/SequencerCommands.cpp
--- a/Engine/Source/Editor/Sequencer/Private/SequencerCommands.cpp
+++ b/Engine/Source/Editor/Sequencer/Private/SequencerCommands.cpp
@@ -38,7 +38,7 @@ void FSequencerCommands::RegisterCommands()
 	UI_COMMAND( SetStartPlaybackRange, "Set Start Playback Range", "Set the start playback range", EUserInterfaceActionType::Button, FInputChord(EKeys::LeftBracket) );
 	UI_COMMAND( SetEndPlaybackRange, "Set End Playback Range", "Set the end playback range", EUserInterfaceActionType::Button, FInputChord(EKeys::RightBracket) );
 	UI_COMMAND( FocusPlaybackTime, "Focus Playback Time", "Focus the view range on the current playback time without changing zoom level", EUserInterfaceActionType::Button, FInputChord());
-	UI_COMMAND( ResetViewRange, "Reset View Range", "Reset view range to the playback range", EUserInterfaceActionType::Button, FInputChord(EKeys::Home) );
+	UI_COMMAND( ResetViewRange, "Reset View Range", "Reset view range to the playback range", EUserInterfaceActionType::Button, FInputChord(EKeys::A) );
 	UI_COMMAND( ZoomToFit, "Zoom to Fit", "Zoom to Fit", EUserInterfaceActionType::Button, FInputChord(EKeys::F) );
 	UI_COMMAND( ZoomToFitSelectionRange, "Focus Selection Range", "Zoom the view to fit the selection range, or the full time bounds if no selection range is set.", EUserInterfaceActionType::Button, FInputChord() );
 	UI_COMMAND( ZoomInViewRange, "Zoom into the View Range", "Zoom into the view range", EUserInterfaceActionType::Button, FInputChord(EKeys::Equals) );
@@ -223,6 +223,12 @@ void FSequencerCommands::RegisterCommands()
 	UI_COMMAND(ToggleTrackSelectionDeactive, "Toggle Track Deactive", "Toggle the deactive state of Sequencer track selection", EUserInterfaceActionType::Button, FInputChord());
 	UI_COMMAND(ToggleTrackSelectionMute, "Toggle Track Mute", "Toggle the muted state of Sequencer track selection", EUserInterfaceActionType::Button, FInputChord());
 	UI_COMMAND(ToggleTrackSelectionSolo, "Toggle Track Solo", "Toggle the soloed state of Sequencer track selection", EUserInterfaceActionType::Button, FInputChord());
+
+	UI_COMMAND(ToggleToolBarVisible, "Toggle Tool Bar Visible", "Toggle between the toolbar panel being visible and hidden", EUserInterfaceActionType::Button, FInputChord());
+	UI_COMMAND(ToggleTransportPanelVisible, "Toggle Transport Panel Visible", "Toggle between the transport controls panel being visible and hidden", EUserInterfaceActionType::Button, FInputChord());
+	UI_COMMAND(ToggleTransportPanelDocked, "Toggle Transport Panel Docked", "Toggle between the transport controls panel being docked within its inline content or undocked in its own tab window", EUserInterfaceActionType::Button, FInputChord());
+
+	UI_COMMAND(ToggleSimpleView, "Toggle Simple/Full View", "Toggle between Sequencer simplified view and full views", EUserInterfaceActionType::Button, FInputChord(EModifierKey::Control, EKeys::BackSpace))
 }
 
 #undef LOCTEXT_NAMESPACE
diff --git a/Engine/Source/Editor/Sequencer/Private/SequencerModule.cpp b/Engine/Source/Editor/Sequencer/Private/SequencerModule.cpp
--- a/Engine/Source/Editor/Sequencer/Private/SequencerModule.cpp
+++ b/Engine/Source/Editor/Sequencer/Private/SequencerModule.cpp
@@ -14,6 +14,7 @@
 #include "SequencerCustomizationManager.h"
 #include "SequencerEdMode.h"
 #include "SequencerObjectChangeListener.h"
+#include "SimpleView/Menus/SequencerToolbarExtension.h"
 #include "IDetailKeyframeHandler.h"
 #include "IDetailTreeNode.h"
 #include "IDetailsView.h"
@@ -457,6 +458,8 @@ public:
 				}
 			}
 		}));
+
+		UE::Sequencer::SimpleView::FSequencerToolbarExtension::AddExtension();
 	}
 
 	void RegisterObjectSchemas()
@@ -546,6 +549,8 @@ public:
 
 	virtual void ShutdownModule() override
 	{
+		UE::Sequencer::SimpleView::FSequencerToolbarExtension::RemoveExtension();
+
 		if (GIsEditor)
 		{
 			UMovieSceneSignedObject::SetDeferredHandler(nullptr);
diff --git a/Engine/Source/Editor/Sequencer/Private/SequencerSettings.cpp b/Engine/Source/Editor/Sequencer/Private/SequencerSettings.cpp
--- a/Engine/Source/Editor/Sequencer/Private/SequencerSettings.cpp
+++ b/Engine/Source/Editor/Sequencer/Private/SequencerSettings.cpp
@@ -1435,3 +1435,12 @@ void USequencerSettings::SetTransportControlsPoppedOut(const bool bInPoppedOut)
 		SaveConfig();
 	}
 }
+
+void USequencerSettings::SetIsSimpleView(const bool bInSimpleView)
+{
+	if (bSimpleView != bInSimpleView)
+	{
+		bSimpleView = bInSimpleView;
+		SaveConfig();
+	}
+}
