[Sequencer] Simple View with toolable timeline initial release

diff --git a/Engine/Source/Editor/Sequencer/Private/SequencerTimeSliderController.cpp b/Engine/Source/Editor/Sequencer/Private/SequencerTimeSliderController.cpp
--- a/Engine/Source/Editor/Sequencer/Private/SequencerTimeSliderController.cpp
+++ b/Engine/Source/Editor/Sequencer/Private/SequencerTimeSliderController.cpp
@@ -77,16 +77,16 @@ FAutoConsoleVariableRef CVarSequencerShowTimeWarpScrubberLink(
 
 FSequencerTimeSliderController::FSequencerTimeSliderController( const FTimeSliderArgs& InArgs, TWeakPtr<FSequencer> InWeakSequencer )
 	: WeakSequencer(InWeakSequencer)
-	, TimeSliderArgs( InArgs )
-	, DistanceDragged( 0.0f )
-	, MouseDragType( DRAG_NONE )
+	, TimeSliderArgs(InArgs)
+	, DistanceDragged(0.0f)
+	, MouseDragType(DRAG_NONE)
 	, bMouseDownInRegion(false)
-	, bPanning( false )
-	, HoverMarkIndex( INDEX_NONE )
+	, bPanning(false)
+	, HoverMarkIndex(INDEX_NONE)
 {
-	ScrubFillBrush              = FAppStyle::GetBrush( TEXT( "Sequencer.Timeline.ScrubFill" ) );
-	FrameBlockScrubHandleDownBrush = FAppStyle::GetBrush( TEXT( "Sequencer.Timeline.FrameBlockScrubHandle" ) );
-	VanillaScrubHandleDownBrush    = FAppStyle::GetBrush( TEXT( "Sequencer.Timeline.ScrubHandle" ) );
+	ScrubFillBrush                 = FAppStyle::GetBrush(TEXT("Sequencer.Timeline.ScrubFill"));
+	FrameBlockScrubHandleDownBrush = FAppStyle::GetBrush(TEXT("Sequencer.Timeline.FrameBlockScrubHandle"));
+	VanillaScrubHandleDownBrush    = FAppStyle::GetBrush(TEXT("Sequencer.Timeline.ScrubHandle"));
 
 	FontMeasureService = FSlateApplication::Get().GetRenderer()->GetFontMeasureService().ToSharedPtr();
 	SmallLayoutFont = FCoreStyle::GetDefaultFontStyle("Regular", 10);
@@ -99,7 +99,6 @@ FSequencerTimeSliderController::FSequencerTimeSliderController( const FTimeSlide
 	{
 		Sequencer->OnGlobalTimeChanged().AddRaw(this, &FSequencerTimeSliderController::SetIsEvaluating);
 	}
-
 }
 
 FSequencerTimeSliderController::~FSequencerTimeSliderController()
@@ -201,17 +200,14 @@ FSequencerTimeSliderController::FScrubberMetrics FSequencerTimeSliderController:
 
 FSequencerTimeSliderController::FScrubberMetrics FSequencerTimeSliderController::GetScrubPixelMetrics(const FQualifiedFrameTime& ScrubTime, const FScrubRangeToScreen& RangeToScreen, float DilationPixels) const
 {
-	FFrameRate DisplayRate = GetDisplayRate();
-
-	FScrubberMetrics Metrics;
-
-	static float MinScrubSize = 14.f;
+	static constexpr float MinScrubSize = 14.f;
+	static constexpr float HalfMinScrubSize = MinScrubSize * .5f;
 
+	const FFrameRate DisplayRate = GetDisplayRate();
 	const FFrameNumber Frame = ScrubTime.ConvertTo(DisplayRate).FloorToFrame();
 
-	float FrameStartPixel = RangeToScreen.InputToLocalX(  Frame    / DisplayRate );
-	float FrameEndPixel   = RangeToScreen.InputToLocalX( (Frame+1) / DisplayRate ) - 1;
-
+	float FrameStartPixel = RangeToScreen.InputToLocalX(Frame / DisplayRate);
+	float FrameEndPixel = RangeToScreen.InputToLocalX((Frame + 1) / DisplayRate) - 1.f;
 	{
 		float RoundedStartPixel = FMath::RoundToInt(FrameStartPixel);
 		FrameEndPixel -= (FrameStartPixel - RoundedStartPixel);
@@ -221,22 +217,32 @@ FSequencerTimeSliderController::FScrubberMetrics FSequencerTimeSliderController:
 	}
 
 	// Store off the pixel width of the frame
-	Metrics.FrameExtentsPx = TRange<float>(FrameStartPixel - DilationPixels, FrameEndPixel + DilationPixels);
+	FScrubberMetrics Metrics;
+	Metrics.FrameExtentsPx = TRange<float>(
+		FrameStartPixel - DilationPixels,
+		FrameEndPixel + DilationPixels
+	);
 
 	// Set the style of the scrub handle
-	TSharedPtr<FSequencer> Sequencer = WeakSequencer.Pin();
-	Metrics.Style = Sequencer.IsValid() ? Sequencer->GetScrubStyle() : ESequencerScrubberStyle::Vanilla;
+	const TSharedPtr<FSequencer> Sequencer = WeakSequencer.Pin();
+	Metrics.Style = Sequencer.IsValid()
+		? Sequencer->GetScrubStyle()
+		: ESequencerScrubberStyle::Vanilla;
 
 	// Always draw the extents on the section area for frame block styles
 	Metrics.bDrawExtents = Metrics.Style == ESequencerScrubberStyle::FrameBlock;
 
 	// If it's vanilla style or too small to show the frame width, set that up
-	if (Metrics.Style == ESequencerScrubberStyle::Vanilla || FrameEndPixel - FrameStartPixel < MinScrubSize)
+	if (Metrics.Style == ESequencerScrubberStyle::Vanilla
+		|| FrameEndPixel - FrameStartPixel < MinScrubSize)
 	{
 		Metrics.Style = ESequencerScrubberStyle::Vanilla;
 
-		float ScrubPixel = RangeToScreen.InputToLocalX(ScrubTime.AsSeconds());
-		Metrics.HandleRangePx = TRange<float>(ScrubPixel - MinScrubSize*.5f - DilationPixels, ScrubPixel + MinScrubSize*.5f + DilationPixels);
+		const float ScrubPixel = RangeToScreen.InputToLocalX(ScrubTime.AsSeconds());
+		Metrics.HandleRangePx = TRange<float>(
+			ScrubPixel - HalfMinScrubSize - DilationPixels,
+			ScrubPixel + HalfMinScrubSize + DilationPixels
+		);
 	}
 	else
 	{
@@ -246,7 +252,7 @@ FSequencerTimeSliderController::FScrubberMetrics FSequencerTimeSliderController:
 	return Metrics;
 }
 
-void FSequencerTimeSliderController::DrawTicks( FSlateWindowElementList& OutDrawElements, const TRange<double>& ViewRange, const FScrubRangeToScreen& RangeToScreen, FDrawTickArgs& InArgs ) const
+void FSequencerTimeSliderController::DrawTicks(FSlateWindowElementList& OutDrawElements, const TRange<double>& ViewRange, const FScrubRangeToScreen& RangeToScreen, FDrawTickArgs& InArgs) const
 {
 	using namespace UE::Sequencer;
 
@@ -361,8 +367,7 @@ void FSequencerTimeSliderController::DrawLinearTicks(FSlateWindowElementList& Ou
 	}
 }
 
-void FSequencerTimeSliderController::DrawHorizontalTrackSeparatorLines(FSlateWindowElementList& OutDrawElements
-	, const TRange<double>& InViewRange, const FScrubRangeToScreen& InRangeToScreen, FDrawTickArgs& InArgs) const
+void FSequencerTimeSliderController::DrawHorizontalTrackSeparatorLines(FSlateWindowElementList& OutDrawElements, FDrawTickArgs& InArgs) const
 {
 	using namespace UE::Sequencer;
 
@@ -395,23 +400,34 @@ void FSequencerTimeSliderController::DrawHorizontalTrackSeparatorLines(FSlateWin
 
 	for (const TViewModelPtr<IOutlinerExtension>& Item : VisibleItems)
 	{
-		if (Item.AsModel()->IsA<FOutlinerSpacer>())
+		FViewModelPtr ItemModel = Item.AsModel();
+
+		if (!ItemModel.IsValid() || ItemModel->IsA<FOutlinerSpacer>())
 		{
 			continue;
 		}
 
 		// Get the vertical position and height of this item
-		const TViewModelPtr<IGeometryExtension> GeometryExtension = Item.AsModel()->CastThisShared<IGeometryExtension>();
+		const TViewModelPtr<IGeometryExtension> GeometryExtension = ItemModel->CastThisShared<IGeometryExtension>();
 		if (!GeometryExtension.IsValid())
 		{
 			continue;
 		}
 
+		if (const TSharedPtr<FViewModel> NextSibling = ItemModel->GetNextSibling())
+		{
+			const TViewModelPtr<IGeometryExtension> NextSiblingGeometry = NextSibling->CastThisShared<IGeometryExtension>();
+			if (!NextSiblingGeometry.IsValid())
+			{
+				continue;
+			}
+		}
+
 		// Convert virtual geometry to physical screen space
 		const FVirtualGeometry ItemGeometry = GeometryExtension->GetVirtualGeometry();
 		const float ItemBottom = TreeViewWidget->VirtualToPhysical(ItemGeometry.GetTop() + ItemGeometry.GetNestedHeight());
 
-		// Draw a horizontal line at the top of this item
+		// Draw a horizontal line at the bottom of this item
 		TArray<FVector2D> LinePoints;
 		LinePoints.AddUninitialized(2);
 		LinePoints[0] = FVector2D(0.f, ItemBottom);
@@ -479,28 +495,44 @@ int32 FSequencerTimeSliderController::DrawScalingAnchors(const FGeometry& Allott
 	return LayerId + 1;
 }
 
-int32 FSequencerTimeSliderController::DrawMarkedFrames( const FGeometry& AllottedGeometry, const FScrubRangeToScreen& RangeToScreen, FSlateWindowElementList& OutDrawElements, int32 LayerId, const ESlateDrawEffect& DrawEffects, const FWidgetStyle& InWidgetStyle, bool bDrawLabels ) const
+int32 FSequencerTimeSliderController::DrawMarkedFrames(const FGeometry& InGeometry,
+	const FScrubRangeToScreen& InRangeToScreen,
+	FSlateWindowElementList& OutDrawElements,
+	int32 LayerId,
+	const ESlateDrawEffect& InDrawEffects,
+	bool bInDrawLabels) const
 {
-	if (!WeakSequencer.Pin()->GetSequencerSettings()->GetShowMarkedFrames())
+	const TSharedPtr<ISequencer> Sequencer = WeakSequencer.Pin();
+	if (!Sequencer.IsValid())
+	{
+		return LayerId;
+	}
+
+	const USequencerSettings* const SequencerSettings = Sequencer->GetSequencerSettings();
+    if (!SequencerSettings)
+    {
+    	return LayerId;
+    }
+
+	if (!SequencerSettings->GetShowMarkedFrames())
 	{
 		return LayerId;
 	}
 
 	using namespace UE::Sequencer;
 
-	const TArray<FMovieSceneMarkedFrame> & MarkedFrames = TimeSliderArgs.MarkedFrames.Get();
-	const TArray<FMovieSceneMarkedFrame> & GlobalMarkedFrames = TimeSliderArgs.GlobalMarkedFrames.Get();
+	const TArray<FMovieSceneMarkedFrame>& MarkedFrames = TimeSliderArgs.MarkedFrames.Get();
+	const TArray<FMovieSceneMarkedFrame>& GlobalMarkedFrames = TimeSliderArgs.GlobalMarkedFrames.Get();
 	if (MarkedFrames.Num() < 1 && GlobalMarkedFrames.Num() < 1)
 	{
 		return LayerId;
 	}
 
-	const FMarkedFrameSelection& SelectedMarkedFrames = WeakSequencer.Pin()->GetViewModel()->GetSelection()->MarkedFrames;
-	//const FLinearColor SelectedColor = FLinearColor::White; //FAppStyle::GetSlateColor("SelectionColor").GetColor(InWidgetStyle);
-	const FLinearColor WhiteColorHSV = FLinearColor::White.LinearRGBToHSV();
-	const FLinearColor DefaultMarkedFrameColor = WeakSequencer.Pin()->GetSequencerSettings()->GetMarkedFrameColor();
+	const FMarkedFrameSelection& SelectedMarkedFrames = Sequencer->GetViewModel()->GetSelection()->MarkedFrames;
+	const FLinearColor DefaultMarkedFrameColor = SequencerSettings->GetMarkedFrameColor();
 
-	auto DrawFrameMarkers = ([=, this](const TArray<FMovieSceneMarkedFrame> & InMarkedFrames, FSlateWindowElementList& DrawElements, bool bIsGlobal)
+	auto DrawFrameMarkers = ([=, this](const TArray<FMovieSceneMarkedFrame>& InMarkedFrames,
+		FSlateWindowElementList& DrawElements, bool bIsGlobal)
 	{
 		for (int32 MarkIndex = 0; MarkIndex < InMarkedFrames.Num(); ++MarkIndex)
 		{
@@ -523,46 +555,49 @@ int32 FSequencerTimeSliderController::DrawMarkedFrames( const FGeometry& Allotte
 				0.5f)
 				.HSVToLinearRGB();
 
-			FLinearColor DrawColor = bIsGlobal ? MarkedFrameColor.Desaturate(0.25f) : (bIsSelected ? SelectedColor : MarkedFrameColor);
-			const float  LinePos = RangeToScreen.InputToLocalX(Seconds);
+			const FLinearColor DrawColor = bIsGlobal
+				? MarkedFrameColor.Desaturate(0.25f)
+				: (bIsSelected ? SelectedColor : MarkedFrameColor);
+			const float LinePos = InRangeToScreen.InputToLocalX(Seconds);
+
 			TArray<FVector2D> LinePoints;
 			LinePoints.AddUninitialized(2);
 			LinePoints[0] = FVector2D(LinePos, 0.0f);
-			LinePoints[1] = FVector2D(LinePos, FMath::FloorToFloat(AllottedGeometry.Size.Y));
+			LinePoints[1] = FVector2D(LinePos, FMath::FloorToFloat(InGeometry.Size.Y));
 
 			FSlateDrawElement::MakeLines(
 				DrawElements,
 				LayerId + 1,
-				AllottedGeometry.ToPaintGeometry(),
+				InGeometry.ToPaintGeometry(),
 				LinePoints,
-				DrawEffects,
+				InDrawEffects,
 				DrawColor,
 				false,
 				(bIsHovered || bIsSelected) ? 2.f : 1.f
 			);
 
 			FString LabelString = MarkedFrame.Label;
-			if (bDrawLabels && !LabelString.IsEmpty())
+			if (bInDrawLabels && !LabelString.IsEmpty())
 			{
 				// Draw the label next to the marked frame line
 				bool bDrawLeft;
 				FVector2D TextPosition, TextSize;
-				GetMarkLabelGeometry(AllottedGeometry, RangeToScreen, MarkedFrame, TextPosition, TextSize, bDrawLeft);
+				GetMarkLabelGeometry(InGeometry, InRangeToScreen, MarkedFrame, TextPosition, TextSize, bDrawLeft);
 
-				const FSlateBrush* LabelBrush = bDrawLeft ?
-					FAppStyle::GetBrush("Sequencer.MarkedFrame.LabelLeft") :
-					FAppStyle::GetBrush("Sequencer.MarkedFrame.LabelRight");
+				const FSlateBrush* LabelBrush = bDrawLeft
+					? FAppStyle::GetBrush(TEXT("Sequencer.MarkedFrame.LabelLeft"))
+					: FAppStyle::GetBrush(TEXT("Sequencer.MarkedFrame.LabelRight"));
 
 				if (bIsHovered || bIsSelected)
 				{
 					FSlateDrawElement::MakeBox(
 						DrawElements,
 						LayerId + 1,
-						AllottedGeometry.ToPaintGeometry(
+						InGeometry.ToPaintGeometry(
 							TextSize + ScrubConstants::MarkLabelBoxMargin + ScrubConstants::MarkLabelBoxWideMargin,
 							FSlateLayoutTransform(FVector2D(LinePos, 0.f))),
 						LabelBrush,
-						DrawEffects,
+						InDrawEffects,
 						DrawColor.CopyWithNewOpacity(0.3f)
 					);
 				}
@@ -570,10 +605,10 @@ int32 FSequencerTimeSliderController::DrawMarkedFrames( const FGeometry& Allotte
 				FSlateDrawElement::MakeText(
 					DrawElements,
 					LayerId + 1,
-					AllottedGeometry.ToPaintGeometry(TextSize, FSlateLayoutTransform(TextPosition)),
+					InGeometry.ToPaintGeometry(TextSize, FSlateLayoutTransform(TextPosition)),
 					LabelString,
 					(bIsHovered || bIsSelected) ? SmallBoldLayoutFont : SmallLayoutFont,
-					DrawEffects,
+					InDrawEffects,
 					DrawColor
 				);
 			}
@@ -725,7 +760,7 @@ int32 FSequencerTimeSliderController::OnPaintTimeSlider( bool bMirrorLabels, con
 			ScrubColor
 		);
 
-		LayerId = DrawMarkedFrames(AllottedGeometry, RangeToScreen, OutDrawElements, LayerId, DrawEffects, InWidgetStyle, true);
+		LayerId = DrawMarkedFrames(AllottedGeometry, RangeToScreen, OutDrawElements, LayerId, DrawEffects, true);
 
 		{
 			// Draw the scrub position string and the local frame string as ScrubPosition : LocalFrame where the ScrubPosition might be time in a parent sequence space
@@ -818,116 +853,165 @@ int32 FSequencerTimeSliderController::OnPaintTimeSlider( bool bMirrorLabels, con
 	return LayerId;
 }
 
-
-int32 FSequencerTimeSliderController::DrawSelectionRange(const FGeometry& AllottedGeometry, const FSlateRect& MyCullingRect, FSlateWindowElementList& OutDrawElements, int32 LayerId, const FScrubRangeToScreen& RangeToScreen, const FPaintPlaybackRangeArgs& Args) const
+int32 FSequencerTimeSliderController::DrawSelectionRange(const FGeometry& InGeometry,
+	const FSlateRect& InCullingRect,
+	FSlateWindowElementList& OutDrawElements,
+	int32 LayerId,
+	const FScrubRangeToScreen& InRangeToScreen,
+	const FPaintPlaybackRangeArgs& InRangeArgs) const
 {
-	TSharedPtr<FSequencer> Sequencer = WeakSequencer.Pin();
-	if (!Sequencer.IsValid())
+	const TRange<double> SelectionRange = TimeSliderArgs.SelectionRange.Get() / GetTickResolution();
+	if (SelectionRange.IsEmpty())
 	{
 		return LayerId;
 	}
 
-	TRange<double> SelectionRange = TimeSliderArgs.SelectionRange.Get() / GetTickResolution();
-
-	if (!SelectionRange.IsEmpty())
-	{
-		const float SelectionRangeL = RangeToScreen.InputToLocalX(SelectionRange.GetLowerBoundValue());
-		const float SelectionRangeR = RangeToScreen.InputToLocalX(SelectionRange.GetUpperBoundValue()) - 1;
-		const auto DrawColor = FAppStyle::GetSlateColor("SelectionColor").GetColor(FWidgetStyle());
-
-		if (Args.SolidFillOpacity > 0.f)
-		{
-			FSlateDrawElement::MakeBox(
-				OutDrawElements,
-				LayerId + 1,
-				AllottedGeometry.ToPaintGeometry(FVector2f(SelectionRangeR - SelectionRangeL, AllottedGeometry.Size.Y), FSlateLayoutTransform(FVector2f(SelectionRangeL, 0.f))),
-				FAppStyle::GetBrush("WhiteBrush"),
-				ESlateDrawEffect::None,
-				DrawColor.CopyWithNewOpacity(Args.SolidFillOpacity)
-			);
-		}
+	const float SelectionRangeL = InRangeToScreen.InputToLocalX(SelectionRange.GetLowerBoundValue());
+	const float SelectionRangeR = InRangeToScreen.InputToLocalX(SelectionRange.GetUpperBoundValue()) - 1.f;
 
-		FSlateDrawElement::MakeBox(
-			OutDrawElements,
-			LayerId + 1,
-			AllottedGeometry.ToPaintGeometry(FVector2f(Args.BrushWidth, AllottedGeometry.Size.Y), FSlateLayoutTransform(FVector2f(SelectionRangeL, 0.f))),
-			Args.StartBrush,
-			ESlateDrawEffect::None,
-			DrawColor
-		);
+	const FLinearColor DrawColor = FAppStyle::GetSlateColor(TEXT("SelectionColor")).GetColor(FWidgetStyle());
+	const float SelectionWidth = FMath::Max(0.f, SelectionRangeR - SelectionRangeL);
 
+	if (InRangeArgs.SolidFillOpacity > 0.f && SelectionWidth > 0.f)
+	{
 		FSlateDrawElement::MakeBox(
 			OutDrawElements,
 			LayerId + 1,
-			AllottedGeometry.ToPaintGeometry(FVector2f(Args.BrushWidth, AllottedGeometry.Size.Y), FSlateLayoutTransform(FVector2f(SelectionRangeR - Args.BrushWidth, 0.f))),
-			Args.EndBrush,
+			InGeometry.ToPaintGeometry(
+				FVector2f(SelectionWidth, InGeometry.Size.Y),
+				FSlateLayoutTransform(FVector2f(SelectionRangeL, 0.f))
+			),
+			FAppStyle::GetBrush(TEXT("WhiteBrush")),
 			ESlateDrawEffect::None,
-			DrawColor
+			DrawColor.CopyWithNewOpacity(InRangeArgs.SolidFillOpacity)
 		);
 	}
 
+	const float BrushWidth = FMath::Max(1.f, InRangeArgs.BrushWidth);
+	const float EndBrushX = FMath::Max(SelectionRangeL, SelectionRangeR - BrushWidth);
+
+	FSlateDrawElement::MakeBox(
+		OutDrawElements,
+		LayerId + 1,
+		InGeometry.ToPaintGeometry(
+			FVector2f(BrushWidth, InGeometry.Size.Y),
+			FSlateLayoutTransform(FVector2f(SelectionRangeL, 0.f))
+		),
+		InRangeArgs.StartBrush,
+		ESlateDrawEffect::None,
+		DrawColor
+	);
+
+	FSlateDrawElement::MakeBox(
+		OutDrawElements,
+		LayerId + 1,
+		InGeometry.ToPaintGeometry(
+			FVector2f(BrushWidth, InGeometry.Size.Y),
+			FSlateLayoutTransform(FVector2f(EndBrushX, 0.f))
+		),
+		InRangeArgs.EndBrush,
+		ESlateDrawEffect::None,
+		DrawColor
+	);
+
 	return LayerId + 1;
 }
 
-
-int32 FSequencerTimeSliderController::DrawPlaybackRange(const FGeometry& AllottedGeometry, const FSlateRect& MyCullingRect, FSlateWindowElementList& OutDrawElements, int32 LayerId, const FScrubRangeToScreen& RangeToScreen, const FPaintPlaybackRangeArgs& Args) const
+int32 FSequencerTimeSliderController::DrawPlaybackRange(const FGeometry& InGeometry,
+	const FSlateRect& InCullingRect,
+	FSlateWindowElementList& OutDrawElements,
+	int32 LayerId,
+	const FScrubRangeToScreen& InRangeToScreen,
+	const FPaintPlaybackRangeArgs& InRangeArgs) const
 {
-	TSharedPtr<FSequencer> Sequencer = WeakSequencer.Pin();
+	const TSharedPtr<FSequencer> Sequencer = WeakSequencer.Pin();
 	if (!Sequencer.IsValid())
 	{
 		return LayerId;
 	}
 
+	const USequencerSettings* const SequencerSettings = Sequencer->GetSequencerSettings();
+	if (!SequencerSettings)
+	{
+		return LayerId;
+	}
+
 	if (!TimeSliderArgs.PlaybackRange.IsSet())
 	{
 		return LayerId;
 	}
 
-	const uint8 OpacityBlend = TimeSliderArgs.SubSequenceRange.Get().IsSet() ? 128 : 255;
+	const TRange<FFrameNumber> PlaybackRange = TimeSliderArgs.PlaybackRange.Get();
+	const FFrameRate TickResolution = GetTickResolution();
 
-	TRange<FFrameNumber> PlaybackRange = TimeSliderArgs.PlaybackRange.Get();
-	FFrameRate TickResolution = GetTickResolution();
-	const float PlaybackRangeL = RangeToScreen.InputToLocalX(PlaybackRange.GetLowerBoundValue() / TickResolution);
-	const float PlaybackRangeR = RangeToScreen.InputToLocalX(PlaybackRange.GetUpperBoundValue() / TickResolution) - 1;
+	float PlaybackRangeL = InRangeToScreen.InputToLocalX(PlaybackRange.GetLowerBoundValue() / TickResolution);
+	float PlaybackRangeR = InRangeToScreen.InputToLocalX(PlaybackRange.GetUpperBoundValue() / TickResolution) - 1.f;
 
-	const FLinearColor PlaybackRangeStartColor = Sequencer->GetSequencerSettings()->GetPlaybackRangeStartColor();
-	const FLinearColor PlaybackRangeEndColor = Sequencer->GetSequencerSettings()->GetPlaybackRangeEndColor();
+	const float BrushWidth = FMath::Max(1.f, InRangeArgs.BrushWidth);
+	const float EndBrushX = FMath::Max(PlaybackRangeL, PlaybackRangeR - BrushWidth);
 
-	FSlateDrawElement::MakeBox(
-		OutDrawElements,
-		LayerId+1,
-		AllottedGeometry.ToPaintGeometry(FVector2f(Args.BrushWidth, AllottedGeometry.Size.Y), FSlateLayoutTransform(FVector2f(PlaybackRangeL, 0.f))),
-		Args.StartBrush,
-		ESlateDrawEffect::None,
-		PlaybackRangeStartColor.CopyWithNewOpacity(OpacityBlend)
-	);
+	const FLinearColor PlaybackRangeStartColor = SequencerSettings->GetPlaybackRangeStartColor();
+	const FLinearColor PlaybackRangeEndColor = SequencerSettings->GetPlaybackRangeEndColor();
+	const uint8 OpacityBlend = TimeSliderArgs.SubSequenceRange.Get().IsSet() ? 128 : 255;
 
-	FSlateDrawElement::MakeBox(
-		OutDrawElements,
-		LayerId+1,
-		AllottedGeometry.ToPaintGeometry(FVector2f(Args.BrushWidth, AllottedGeometry.Size.Y), FSlateLayoutTransform(FVector2f(PlaybackRangeR - Args.BrushWidth, 0.f))),
-		Args.EndBrush,
-		ESlateDrawEffect::None,
-		PlaybackRangeEndColor.CopyWithNewOpacity(OpacityBlend)
-	);
+	// Only draw start marker if it's within the visible area
+	if (PlaybackRangeL >= 0.f && PlaybackRangeL <= InGeometry.Size.X)
+	{
+		FSlateDrawElement::MakeBox(
+			OutDrawElements,
+			LayerId + 1,
+			InGeometry.ToPaintGeometry(
+				FVector2f(BrushWidth, InGeometry.Size.Y),
+				FSlateLayoutTransform(FVector2f(PlaybackRangeL, 0.f))
+			),
+			InRangeArgs.StartBrush,
+			ESlateDrawEffect::None,
+			PlaybackRangeStartColor.CopyWithNewOpacity(OpacityBlend)
+		);
+	}
+
+	// Only draw end marker if it's within the visible area
+	if (EndBrushX >= 0.f && EndBrushX <= InGeometry.Size.X)
+	{
+		FSlateDrawElement::MakeBox(
+			OutDrawElements,
+			LayerId + 1,
+			InGeometry.ToPaintGeometry(
+				FVector2f(BrushWidth, InGeometry.Size.Y),
+				FSlateLayoutTransform(FVector2f(EndBrushX, 0.f))
+			),
+			InRangeArgs.EndBrush,
+			ESlateDrawEffect::None,
+			PlaybackRangeEndColor.CopyWithNewOpacity(OpacityBlend)
+		);
+	}
 
 	// Black tint for excluded regions
+	const float LeftExcludedRange = FMath::Max(0.f, PlaybackRangeL);
+	const float RightExcludedRange = FMath::Max(0.f, InGeometry.Size.X - PlaybackRangeR);
+
 	FSlateDrawElement::MakeBox(
 		OutDrawElements,
-		LayerId+1,
-		AllottedGeometry.ToPaintGeometry(FVector2f(PlaybackRangeL, AllottedGeometry.Size.Y), FSlateLayoutTransform()),
-		FAppStyle::GetBrush("WhiteBrush"),
+		LayerId + 1,
+		InGeometry.ToPaintGeometry(
+			FVector2f(LeftExcludedRange, InGeometry.Size.Y),
+			FSlateLayoutTransform()
+		),
+		FAppStyle::GetBrush(TEXT("WhiteBrush")),
 		ESlateDrawEffect::None,
-		FLinearColor::Black.CopyWithNewOpacity(0.3f * OpacityBlend / 255.f)
+		FLinearColor::Black.CopyWithNewOpacity(.3f * OpacityBlend / 255.f)
 	);
 
 	FSlateDrawElement::MakeBox(
 		OutDrawElements,
-		LayerId+1,
-		AllottedGeometry.ToPaintGeometry(FVector2f(AllottedGeometry.Size.X - PlaybackRangeR, AllottedGeometry.Size.Y), FSlateLayoutTransform(FVector2f(PlaybackRangeR, 0.f))),
-		FAppStyle::GetBrush("WhiteBrush"),
+		LayerId + 1,
+		InGeometry.ToPaintGeometry(
+			FVector2f(RightExcludedRange, InGeometry.Size.Y),
+			FSlateLayoutTransform(FVector2f(PlaybackRangeR, 0.f))
+		),
+		FAppStyle::GetBrush(TEXT("WhiteBrush")),
 		ESlateDrawEffect::None,
-		FLinearColor::Black.CopyWithNewOpacity(0.3f * OpacityBlend / 255.f)
+		FLinearColor::Black.CopyWithNewOpacity(.3f * OpacityBlend / 255.f)
 	);
 
 	return LayerId + 1;
@@ -1023,190 +1107,221 @@ int32 FSequencerTimeSliderController::DrawSubSequenceRange(const FGeometry& Allo
 	return LayerId + 1;
 }
 
-FReply FSequencerTimeSliderController::OnMouseButtonDown( SWidget& WidgetOwner, const FGeometry& MyGeometry, const FPointerEvent& MouseEvent )
+void FSequencerTimeSliderController::ResetMouseDown(const FGeometry& InGeometry, const FPointerEvent& InPointerEvent)
 {
+	const FVector2D PointerPosition = InPointerEvent.GetScreenSpacePosition();
+
 	MouseDragType = DRAG_NONE;
 	DistanceDragged = 0;
 	MouseDownPlaybackRange = TimeSliderArgs.PlaybackRange.Get();
 	MouseDownSelectionRange = TimeSliderArgs.SelectionRange.Get();
-	MouseDownPosition = LastMousePosition = MouseEvent.GetScreenSpacePosition();
-	MouseDownGeometry = MyGeometry;
-	bMouseDownInRegion = false;
+	MouseDownPosition = PointerPosition;
+	LastMousePosition = PointerPosition;
+	MouseDownGeometry = InGeometry;
 
-	FVector2D CursorPos = MouseEvent.GetScreenSpacePosition();
-	FVector2D LocalPos = MouseDownGeometry.AbsoluteToLocal(CursorPos);
-	if (LocalPos.Y >= 0 && LocalPos.Y < MouseDownGeometry.GetLocalSize().Y)
-	{
-		bMouseDownInRegion = true;
-	}
+	const FVector2D LocalPointerPosition = MouseDownGeometry.AbsoluteToLocal(PointerPosition);
+	bMouseDownInRegion = LocalPointerPosition.Y >= 0
+		&& LocalPointerPosition.Y < MouseDownGeometry.GetLocalSize().Y;
+}
+
+FReply FSequencerTimeSliderController::OnMouseButtonDown(SWidget& WidgetOwner, const FGeometry& InGeometry, const FPointerEvent& InPointerEvent)
+{
+	ResetMouseDown(InGeometry, InPointerEvent);
 
 	return FReply::Unhandled();
 }
 
-FReply FSequencerTimeSliderController::OnMouseButtonUp( SWidget& WidgetOwner, const FGeometry& MyGeometry, const FPointerEvent& MouseEvent )
+FReply FSequencerTimeSliderController::OnMouseButtonUp_ContextMenu(SWidget& WidgetOwner, const FGeometry& InGeometry, const FPointerEvent& InPointerEvent)
 {
-	bool bHandleLeftMouseButton  = MouseEvent.GetEffectingButton() == EKeys::LeftMouseButton  && WidgetOwner.HasMouseCapture();
-	bool bHandleRightMouseButton = MouseEvent.GetEffectingButton() == EKeys::RightMouseButton && WidgetOwner.HasMouseCapture() && TimeSliderArgs.AllowZoom ;
-	bool bHandleMiddleMouseButton = MouseEvent.GetEffectingButton() == EKeys::MiddleMouseButton && WidgetOwner.HasMouseCapture();
-	FScrubRangeToScreen RangeToScreen = FScrubRangeToScreen(GetViewRange(), MyGeometry.Size);
-	FFrameTime          MouseTime     = ComputeFrameTimeFromMouse(MyGeometry, MouseEvent.GetScreenSpacePosition(), RangeToScreen);
+	const FKey EffectingButton = InPointerEvent.GetEffectingButton();
+	if (EffectingButton != EKeys::RightMouseButton)
+	{
+		return FReply::Unhandled();
+	}
 
-	if ( bHandleRightMouseButton )
+	if (!bPanning)
 	{
-		if (!bPanning)
+		// Open a context menu if allowed
+		if (ContextMenuSuppression == 0 && DistanceDragged == 0.f)
 		{
-			// Open a context menu if allowed
-			if (ContextMenuSuppression == 0 && TimeSliderArgs.PlaybackRange.IsSet())
+			if (const TSharedPtr<SWidget> MenuContent = OpenContextMenu(InGeometry, InPointerEvent))
 			{
-				TSharedRef<SWidget> MenuContent = OpenSetPlaybackRangeMenu(MyGeometry, MouseEvent);
+				const TSharedRef<SWidget> MenuContentRef = MenuContent.ToSharedRef();
+				const FVector2D PointerPosition = InPointerEvent.GetScreenSpacePosition();
+
 				FSlateApplication::Get().PushMenu(
 					WidgetOwner.AsShared(),
-					MouseEvent.GetEventPath() != nullptr ? *MouseEvent.GetEventPath() : FWidgetPath(),
-					MenuContent,
-					MouseEvent.GetScreenSpacePosition(),
-					FPopupTransitionEffect( FPopupTransitionEffect::ContextMenu )
-					);
-
-				return FReply::Handled().SetUserFocus(MenuContent, EFocusCause::SetDirectly).ReleaseMouseCapture();
-			}
+					InPointerEvent.GetEventPath() ? *InPointerEvent.GetEventPath() : FWidgetPath(),
+					MenuContentRef,
+					PointerPosition,
+					FPopupTransitionEffect(FPopupTransitionEffect::ContextMenu)
+				);
 
-			// return unhandled in case our parent wants to use our right mouse button to open a context menu
-			if (DistanceDragged == 0.f)
-			{
-				return FReply::Unhandled().ReleaseMouseCapture();
+				return FReply::Handled()
+					.SetUserFocus(MenuContentRef, EFocusCause::SetDirectly)
+					.ReleaseMouseCapture();
 			}
 		}
-		
-		bPanning = false;
-		bMouseDownInRegion = false;
-		MouseDownPosition.Reset();
-		LastMousePosition.Reset();
-		
-		return FReply::Handled().ReleaseMouseCapture();
+
+		// Return unhandled in case our parent wants to use our right mouse button to open a context menu
+		return FReply::Unhandled().ReleaseMouseCapture();
+	}
+
+	ResetMouseInput();
+
+	return FReply::Handled().ReleaseMouseCapture();
+}
+
+FReply FSequencerTimeSliderController::OnMouseButtonUp_MouseDragType(SWidget& WidgetOwner
+	, const FGeometry& InGeometry
+	, const FPointerEvent& InPointerEvent
+	, const FScrubRangeToScreen& InRangeToScreen
+	, const FFrameTime& InMouseTime)
+{
+	if (MouseDragType == DRAG_PLAYBACK_START)
+	{
+		TimeSliderArgs.OnPlaybackRangeEndDrag.ExecuteIfBound();
 	}
-	else if ( bHandleLeftMouseButton || bHandleMiddleMouseButton )
+	else if (MouseDragType == DRAG_PLAYBACK_END)
 	{
-		if (MouseDragType == DRAG_PLAYBACK_START)
-		{
-			TimeSliderArgs.OnPlaybackRangeEndDrag.ExecuteIfBound();
-		}
-		else if (MouseDragType == DRAG_PLAYBACK_END)
-		{
-			TimeSliderArgs.OnPlaybackRangeEndDrag.ExecuteIfBound();
-		}
-		else if (MouseDragType == DRAG_SELECTION_START)
-		{
-			TimeSliderArgs.OnSelectionRangeEndDrag.ExecuteIfBound();
-		}
-		else if (MouseDragType == DRAG_SELECTION_END)
-		{
-			TimeSliderArgs.OnSelectionRangeEndDrag.ExecuteIfBound();
-		}
-		else if (MouseDragType == DRAG_MARK)
-		{
-			TimeSliderArgs.OnMarkEndDrag.ExecuteIfBound();
-		}
-		else if (MouseDragType == DRAG_SETTING_RANGE && MouseDownPosition.IsSet())
-		{
-			// Zooming
-			FFrameTime MouseDownStart = ComputeFrameTimeFromMouse(MyGeometry, MouseDownPosition.GetValue(), RangeToScreen);
+		TimeSliderArgs.OnPlaybackRangeEndDrag.ExecuteIfBound();
+	}
+	else if (MouseDragType == DRAG_SELECTION_START)
+	{
+		TimeSliderArgs.OnSelectionRangeEndDrag.ExecuteIfBound();
+	}
+	else if (MouseDragType == DRAG_SELECTION_END)
+	{
+		TimeSliderArgs.OnSelectionRangeEndDrag.ExecuteIfBound();
+	}
+	else if (MouseDragType == DRAG_MARK)
+	{
+		TimeSliderArgs.OnMarkEndDrag.ExecuteIfBound();
+	}
+	else if (MouseDragType == DRAG_SETTING_RANGE && MouseDownPosition.IsSet())
+	{
+		// Zooming
+		const FFrameTime MouseDownStart = ComputeFrameTimeFromMouse(InGeometry, MouseDownPosition.GetValue(), InRangeToScreen);
 
-			const bool bCanZoomIn  = MouseTime > MouseDownStart;
-			const bool bCanZoomOut = ViewRangeStack.Num() > 0;
-			if (bCanZoomIn || bCanZoomOut)
+		const bool bCanZoomIn  = InMouseTime > MouseDownStart;
+		const bool bCanZoomOut = ViewRangeStack.Num() > 0;
+		if (bCanZoomIn || bCanZoomOut)
+		{
+			TRange<double> ViewRange = GetViewRange();
+			if (!bCanZoomIn)
 			{
-				TRange<double> ViewRange = GetViewRange();
-				if (!bCanZoomIn)
-				{
-					ViewRange = ViewRangeStack.Pop();
-				}
+				ViewRange = ViewRangeStack.Pop();
+			}
 
-				if (bCanZoomIn)
-				{
-					// push the current value onto the stack
-					ViewRangeStack.Add(ViewRange);
+			if (bCanZoomIn)
+			{
+				// push the current value onto the stack
+				ViewRangeStack.Add(ViewRange);
 
-					ViewRange = TRange<double>(MouseDownStart.FrameNumber / GetTickResolution(), MouseTime.FrameNumber / GetTickResolution());
-				}
-				
-				TimeSliderArgs.OnViewRangeChanged.ExecuteIfBound(ViewRange, EViewRangeInterpolation::Immediate);
-				if( !TimeSliderArgs.ViewRange.IsBound() )
-				{
-					// The output is not bound to a delegate so we'll manage the value ourselves
-					TimeSliderArgs.ViewRange.Set(ViewRange);
-				}
+				const FFrameRate TickResolution = GetTickResolution();
+				ViewRange = TRange<double>(MouseDownStart.FrameNumber / TickResolution, InMouseTime.FrameNumber / TickResolution);
+			}
+			
+			TimeSliderArgs.OnViewRangeChanged.ExecuteIfBound(ViewRange, EViewRangeInterpolation::Immediate);
+			if (!TimeSliderArgs.ViewRange.IsBound())
+			{
+				// The output is not bound to a delegate so we'll manage the value ourselves
+				TimeSliderArgs.ViewRange.Set(ViewRange);
 			}
 		}
-		else if (bMouseDownInRegion)
+	}
+	else if (bMouseDownInRegion)
+	{
+		const TSharedPtr<FSequencer> Sequencer = WeakSequencer.Pin();
+		const bool bHandleMiddleMouseButton = InPointerEvent.GetEffectingButton() == EKeys::MiddleMouseButton && WidgetOwner.HasMouseCapture();
+
+		if (HoverMarkIndex == INDEX_NONE)
 		{
-			TSharedPtr<FSequencer> Sequencer = WeakSequencer.Pin();
+			FFrameTime ScrubTime = InMouseTime;
 
-			if (HoverMarkIndex == INDEX_NONE)
+			if (MouseDragType == DRAG_SCRUBBING_TIME)
 			{
-				FFrameTime ScrubTime = MouseTime;
-				FVector2D CursorPos  = MouseEvent.GetScreenSpacePosition();
+				ScrubTime = ComputeScrubTimeFromMouse(InGeometry, InPointerEvent, InRangeToScreen);
+			}
+			else if (Sequencer.IsValid())
+			{
+				const FVector2D PointerPosition  = InPointerEvent.GetScreenSpacePosition();
+				ScrubTime = SnapTimeToNearestKey(InPointerEvent, InRangeToScreen, PointerPosition.X, ScrubTime);
 
-				if (MouseDragType == DRAG_SCRUBBING_TIME)
-				{
-					ScrubTime = ComputeScrubTimeFromMouse(MyGeometry, MouseEvent, RangeToScreen);
-				}
-				else if (Sequencer.IsValid())
-				{
-					ScrubTime = SnapTimeToNearestKey(MouseEvent, RangeToScreen, CursorPos.X, ScrubTime);
+				// We weren't scrubbing, so this is a single click in the time slider. This should
+				// also act as a deselection of any selected markers.
+				Sequencer->GetViewModel()->GetSelection()->MarkedFrames.Empty();
+			}
 
-					// We weren't scrubbing, so this is a single click in the time slider. This should
-					// also act as a deselection of any selected markers.
-					Sequencer->GetViewModel()->GetSelection()->MarkedFrames.Empty();
-				}
+			// If middle mouse button down we don't evaluate on the time change
+			CommitScrubPosition(ScrubTime, /*bIsScrubbing=*/false, /*bEvaluate*/!bHandleMiddleMouseButton);
 
-				// If middle mouse button down we don't evaluate on the time change
-				CommitScrubPosition( ScrubTime, /*bIsScrubbing=*/false , /*bEvaluate*/ !bHandleMiddleMouseButton);
+			// Teleport the playhead to the clicked time.
+			TimeSliderArgs.OnEndScrubberMovement.ExecuteIfBound();
+		}
+		else if (Sequencer.IsValid())
+		{
+			// Select the clicked marker, and optionally teleport the playhead.
+			HandleDragMark(HoverMarkIndex);
 
-				// Teleport the playhead to the clicked time.
-				TimeSliderArgs.OnEndScrubberMovement.ExecuteIfBound();
-			}
-			else if (Sequencer.IsValid())
+			const FModifierKeysState ModifierKeys = FSlateApplication::Get().GetModifierKeys();
+			const bool bSnapPlayTimeToMark = ModifierKeys.AreModifersDown(EModifierKey::Shift);
+			if (bSnapPlayTimeToMark || Sequencer->GetSequencerSettings()->GetSnapPlayTimeToPressedKey())
 			{
-				// Select the clicked marker, and optionally teleport the playhead.
-				HandleMarkSelection(HoverMarkIndex);
+				const TArray<FMovieSceneMarkedFrame>& MarkedFrames = TimeSliderArgs.MarkedFrames.Get();
+				const FFrameNumber MarkTime = MarkedFrames[HoverMarkIndex].FrameNumber;
 
-				FModifierKeysState ModifierKeys = FSlateApplication::Get().GetModifierKeys();
-				const bool bSnapPlayTimeToMark = ModifierKeys.AreModifersDown(EModifierKey::Shift);
-				if (bSnapPlayTimeToMark || Sequencer->GetSequencerSettings()->GetSnapPlayTimeToPressedKey())
-				{
-					const TArray<FMovieSceneMarkedFrame>& MarkedFrames = TimeSliderArgs.MarkedFrames.Get();
-					const FFrameNumber MarkTime = MarkedFrames[HoverMarkIndex].FrameNumber;
-
-					// If middle mouse button down we don't evaluate on the time change
-					CommitScrubPosition(MarkTime, /*bIsScrubbing=*/ false , /*bEvaluate*/ !bHandleMiddleMouseButton);
-				}
+				// If middle mouse button down we don't evaluate on the time change
+				CommitScrubPosition(MarkTime, /*bIsScrubbing=*/false, /*bEvaluate*/!bHandleMiddleMouseButton);
 			}
 		}
+	}
 
-		MouseDragType = DRAG_NONE;
-		DistanceDragged = 0.f;
-		bMouseDownInRegion = false;
-		MouseDownPosition.Reset();
-		LastMousePosition.Reset();
+	MouseDragType = DRAG_NONE;
+	DistanceDragged = 0.f;
+	ResetMouseInput();
 
-		return FReply::Handled().ReleaseMouseCapture();
+	return FReply::Handled().ReleaseMouseCapture();
+}
+
+FReply FSequencerTimeSliderController::OnMouseButtonUp(SWidget& OwnerWidget, const FGeometry& InGeometry, const FPointerEvent& InPointerEvent)
+{
+	const bool bHandleLeftMouseButton  = InPointerEvent.GetEffectingButton() == EKeys::LeftMouseButton  && OwnerWidget.HasMouseCapture();
+	const bool bHandleRightMouseButton = InPointerEvent.GetEffectingButton() == EKeys::RightMouseButton && OwnerWidget.HasMouseCapture() && TimeSliderArgs.AllowZoom;
+	const bool bHandleMiddleMouseButton = InPointerEvent.GetEffectingButton() == EKeys::MiddleMouseButton && OwnerWidget.HasMouseCapture();
+
+	const FScrubRangeToScreen RangeToScreen = FScrubRangeToScreen(GetViewRange(), InGeometry.Size);
+	const FFrameTime MouseTime = ComputeFrameTimeFromMouse(InGeometry, InPointerEvent.GetScreenSpacePosition(), RangeToScreen);
+
+	FReply OutReply = OnMouseButtonUp_ContextMenu(OwnerWidget, InGeometry, InPointerEvent);
+	if (OutReply.IsEventHandled())
+	{
+		return OutReply;
 	}
 
-	bMouseDownInRegion = false;
-	MouseDownPosition.Reset();
-	LastMousePosition.Reset();
+	if (bHandleLeftMouseButton || bHandleMiddleMouseButton)
+	{
+		OutReply = OnMouseButtonUp_MouseDragType(OwnerWidget, InGeometry, InPointerEvent, RangeToScreen, MouseTime);
+		if (OutReply.IsEventHandled())
+		{
+			return OutReply;
+		}
+	}
+
+	ResetMouseInput();
+
 	return FReply::Unhandled();
 }
 
-FReply FSequencerTimeSliderController::OnMouseButtonDoubleClick(SWidget& WidgetOwner, const FGeometry& MyGeometry, const FPointerEvent& MouseEvent)
+FReply FSequencerTimeSliderController::OnMouseButtonDoubleClick(SWidget& OwnerWidget, const FGeometry& InGeometry, const FPointerEvent& InPointerEvent)
 {
 	if (HoverMarkIndex != INDEX_NONE)
 	{
 		// Teleport the playhead to the double-clicked marker.
-		bool bHandleMiddleMouseButton = MouseEvent.GetEffectingButton() == EKeys::MiddleMouseButton && WidgetOwner.HasMouseCapture();
+		const bool bHandleMiddleMouseButton = InPointerEvent.GetEffectingButton() == EKeys::MiddleMouseButton && OwnerWidget.HasMouseCapture();
 		const TArray<FMovieSceneMarkedFrame>& MarkedFrames = TimeSliderArgs.MarkedFrames.Get();
 		const FFrameNumber MarkTime = MarkedFrames[HoverMarkIndex].FrameNumber;
-		CommitScrubPosition(MarkTime, /*bIsScrubbing=*/ false , /*bEvaluate*/ !bHandleMiddleMouseButton);
+		CommitScrubPosition(MarkTime, /*bIsScrubbing=*/false, /*bEvaluate*/!bHandleMiddleMouseButton);
 	}
 	return FReply::Unhandled();
 }
@@ -1227,29 +1342,30 @@ double FSequencerTimeSliderController::ComputeHeight() const
 	return 22.0;
 }
 
-FReply FSequencerTimeSliderController::OnMouseMove( SWidget& WidgetOwner, const FGeometry& MyGeometry, const FPointerEvent& MouseEvent )
+FReply FSequencerTimeSliderController::OnMouseMove(SWidget& OwnerWidget, const FGeometry& InGeometry, const FPointerEvent& InPointerEvent)
 {
-	return OnMouseMoveImpl(WidgetOwner, MyGeometry, MouseEvent, false);
+	return OnMouseMoveImpl(OwnerWidget, InGeometry, InPointerEvent, false);
 }
 
-FReply FSequencerTimeSliderController::OnTimeSliderMouseMove(SWidget& WidgetOwner, const FGeometry& MyGeometry, const FPointerEvent& MouseEvent)
+FReply FSequencerTimeSliderController::OnTimeSliderMouseMove(SWidget& OwnerWidget, const FGeometry& InGeometry, const FPointerEvent& InPointerEvent)
 {
-	return OnMouseMoveImpl(WidgetOwner, MyGeometry, MouseEvent, true);
+	return OnMouseMoveImpl(OwnerWidget, InGeometry, InPointerEvent, true);
 }
 
-FReply FSequencerTimeSliderController::OnMouseMoveImpl( SWidget& WidgetOwner, const FGeometry& MyGeometry, const FPointerEvent& MouseEvent, bool bFromTimeSlider )
+FReply FSequencerTimeSliderController::OnMouseMoveImpl(SWidget& WidgetOwner
+	, const FGeometry& InGeometry
+	, const FPointerEvent& InPointerEvent
+	, const bool bInFromTimeSlider)
 {
-	using namespace UE::Sequencer;
-
-	TSharedPtr<FSequencer> Sequencer = WeakSequencer.Pin();
+	const TSharedPtr<FSequencer> Sequencer = WeakSequencer.Pin();
 	if (!Sequencer.IsValid())
 	{
 		return FReply::Unhandled();
 	}
 
-	bool bHandleLeftMouseButton  = MouseEvent.IsMouseButtonDown( EKeys::LeftMouseButton  );
-	bool bHandleRightMouseButton = MouseEvent.IsMouseButtonDown( EKeys::RightMouseButton ) && TimeSliderArgs.AllowZoom;
-	bool bHandleMiddleMouseButton = MouseEvent.IsMouseButtonDown(EKeys::MiddleMouseButton);
+	const bool bHandleLeftMouseButton = InPointerEvent.IsMouseButtonDown(EKeys::LeftMouseButton);
+	const bool bHandleRightMouseButton = InPointerEvent.IsMouseButtonDown(EKeys::RightMouseButton) && TimeSliderArgs.AllowZoom;
+	const bool bHandleMiddleMouseButton = InPointerEvent.IsMouseButtonDown(EKeys::MiddleMouseButton);
 
 	const bool bLockedMarkedFrames = TimeSliderArgs.AreMarkedFramesLocked.Get();
 
@@ -1258,207 +1374,254 @@ FReply FSequencerTimeSliderController::OnMouseMoveImpl( SWidget& WidgetOwner, co
 
 	if (bHandleRightMouseButton)
 	{
-		if (!bPanning)
+		OnMouseMove_RightMouseDown(InGeometry, InPointerEvent, *Sequencer->GetSequencerSettings());
+	}
+	else if ((bHandleLeftMouseButton || bHandleMiddleMouseButton) && MouseDownPosition.IsSet())
+	{
+		const TRange<double> LocalViewRange = GetViewRange();
+		const FScrubRangeToScreen RangeToScreen(LocalViewRange, InGeometry.Size);
+		OnMouseMove_LeftAndMiddleMouseDown(InGeometry, InPointerEvent, bInFromTimeSlider, RangeToScreen, *Sequencer->GetSequencerSettings());
+	}
+	else if (bInFromTimeSlider && !bLockedMarkedFrames && DragMarkIndex == INDEX_NONE)
+	{
+		// Update hover state of marked frames.
+		const TRange<double> LocalViewRange = GetViewRange();
+		const FScrubRangeToScreen RangeToScreen(LocalViewRange, InGeometry.Size);
+		const FVector2D LocalMousePostion = InGeometry.AbsoluteToLocal(InPointerEvent.GetScreenSpacePosition());
+		HitTestMark(InGeometry, RangeToScreen, LocalMousePostion.X, bInFromTimeSlider, &HoverMarkIndex);
+	}
+
+	if (DistanceDragged != 0.f && (bHandleLeftMouseButton || bHandleRightMouseButton))
+	{
+		return FReply::Handled().CaptureMouse(WidgetOwner.AsShared());
+	}
+
+	return FReply::Handled();
+}
+
+FReply FSequencerTimeSliderController::OnMouseMove_LeftAndMiddleMouseDown(const FGeometry& InGeometry
+	, const FPointerEvent& InPointerEvent
+	, const bool bInFromTimeSlider
+	, const FScrubRangeToScreen& InRangeToScreen
+	, const USequencerSettings& InSequencerSettings)
+{
+	const bool bHandleLeftMouseButton = InPointerEvent.IsMouseButtonDown(EKeys::LeftMouseButton);
+	const bool bHandleRightMouseButton = InPointerEvent.IsMouseButtonDown(EKeys::RightMouseButton) && TimeSliderArgs.AllowZoom;
+	const bool bHandleMiddleMouseButton = InPointerEvent.IsMouseButtonDown(EKeys::MiddleMouseButton);
+
+	DistanceDragged += FMath::Abs(InPointerEvent.GetCursorDelta().X);
+
+	if (MouseDragType == DRAG_NONE)
+	{
+		if (DistanceDragged > 0.f/*FSlateApplication::Get().GetDragTriggerDistance()*/)
 		{
-			DistanceDragged += FMath::Abs( MouseEvent.GetCursorDelta().X );
-			if ( DistanceDragged > 0.f /*FSlateApplication::Get().GetDragTriggerDistance()*/ )
+			const FFrameTime MouseDownFree = ComputeFrameTimeFromMouse(InGeometry
+				, MouseDownPosition.GetValue(), InRangeToScreen, false);
+
+			//const bool       bReadOnly          = Sequencer->IsReadOnly();
+			const FFrameRate TickResolution     = GetTickResolution();
+			const bool       bLockedPlayRange   = TimeSliderArgs.IsPlaybackRangeLocked.Get();
+			const float      MouseDownPixel     = InRangeToScreen.InputToLocalX(MouseDownFree / TickResolution);
+			const bool       bHitScrubber       = GetHitTestScrubPixelMetrics(InRangeToScreen).HandleRangePx.Contains(MouseDownPixel);
+
+			const TRange<double> SelectionRange = TimeSliderArgs.SelectionRange.Get() / TickResolution;
+			const TRange<double> PlaybackRange = TimeSliderArgs.PlaybackRange.Get() / TickResolution;
+
+			int32 HitMarkIndex = INDEX_NONE;
+
+			// Disable selection range test if it's empty so that the playback range scrubbing gets priority
+			if (!bHandleMiddleMouseButton
+				&& !bHitScrubber
+				&& !SelectionRange.IsEmpty()
+				&& HitTestRangeEnd(InRangeToScreen, SelectionRange, MouseDownPixel))
 			{
-				bPanning = true;
+				// selection range end scrubber
+				MouseDragType = DRAG_SELECTION_END;
+				TimeSliderArgs.OnSelectionRangeBeginDrag.ExecuteIfBound();
 			}
-		}
-		else if (MouseEvent.IsShiftDown() && MouseEvent.IsAltDown())
-		{
-			float MouseFractionX = MyGeometry.AbsoluteToLocal(MouseEvent.GetScreenSpacePosition()).X / MyGeometry.GetLocalSize().X;
-
-			// If zooming on the current time, adjust mouse fractionX
-			if (Sequencer->GetSequencerSettings()->GetZoomPosition() == ESequencerZoomPosition::SZP_CurrentTime)
+			else if (!bHandleMiddleMouseButton
+				&& !bHitScrubber
+				&& !SelectionRange.IsEmpty()
+				&& HitTestRangeStart(InRangeToScreen, SelectionRange, MouseDownPixel))
 			{
-				const double ScrubPosition = TimeSliderArgs.ScrubPosition.Get() / GetTickResolution();
-				if (GetViewRange().Contains(ScrubPosition))
-				{
-					FScrubRangeToScreen RangeToScreen(GetViewRange(), MyGeometry.Size);
-					float TimePosition = RangeToScreen.InputToLocalX(ScrubPosition);
-					MouseFractionX = TimePosition / MyGeometry.GetLocalSize().X;
-				}
+				// selection range start scrubber
+				MouseDragType = DRAG_SELECTION_START;
+				TimeSliderArgs.OnSelectionRangeBeginDrag.ExecuteIfBound();
 			}
+			else if (!bHandleMiddleMouseButton
+				&& !bHitScrubber
+				&& !bLockedPlayRange
+				&& HitTestRangeEnd(InRangeToScreen, PlaybackRange, MouseDownPixel))
+			{
+				// playback range end scrubber
+				MouseDragType = DRAG_PLAYBACK_END;
+				TimeSliderArgs.OnPlaybackRangeBeginDrag.ExecuteIfBound();
+			}
+			else if (!bHandleMiddleMouseButton
+				&& !bHitScrubber
+				&& !bLockedPlayRange
+				&& HitTestRangeStart(InRangeToScreen, PlaybackRange, MouseDownPixel))
+			{
+				// playback range start scrubber
+				MouseDragType = DRAG_PLAYBACK_START;
+				TimeSliderArgs.OnPlaybackRangeBeginDrag.ExecuteIfBound();
+			}
+			else if (!bHandleMiddleMouseButton
+				&& !bHitScrubber
+				&& !TimeSliderArgs.AreMarkedFramesLocked.Get()
+				&& HitTestMark(InGeometry, InRangeToScreen, MouseDownPixel, bInFromTimeSlider, &HitMarkIndex))
+			{
+				MouseDragType = DRAG_MARK;
 
-			const float ZoomDelta = -0.01f * MouseEvent.GetCursorDelta().X;
-			ZoomByDelta(ZoomDelta, MouseFractionX);
-		}
-		else 
-		{
-			TRange<double> LocalViewRange = GetViewRange();
-			double LocalViewRangeMin = LocalViewRange.GetLowerBoundValue();
-			double LocalViewRangeMax = LocalViewRange.GetUpperBoundValue();
-
-			FScrubRangeToScreen ScaleInfo( LocalViewRange, MyGeometry.Size );
-			FVector2D ScreenDelta = MouseEvent.GetCursorDelta();
-			FVector2D InputDelta;
-			InputDelta.X = ScreenDelta.X/ScaleInfo.PixelsPerInput;
-
-			double NewViewOutputMin = LocalViewRangeMin - InputDelta.X;
-			double NewViewOutputMax = LocalViewRangeMax - InputDelta.X;
+				HandleDragMark_NoDragType(HitMarkIndex);
 
-			ClampViewRange(NewViewOutputMin, NewViewOutputMax);
-			SetViewRange(NewViewOutputMin, NewViewOutputMax, EViewRangeInterpolation::Immediate);
+				TimeSliderArgs.OnMarkBeginDrag.ExecuteIfBound();
+			}
+			else if (FSlateApplication::Get().GetModifierKeys().AreModifersDown(EModifierKey::Control) && bHandleMiddleMouseButton == false)
+			{
+				MouseDragType = DRAG_SETTING_RANGE;
+			}
+			else if (bMouseDownInRegion)
+			{
+				MouseDragType = DRAG_SCRUBBING_TIME;
+				TimeSliderArgs.OnBeginScrubberMovement.ExecuteIfBound();
+			}
 		}
 	}
-	else if ((bHandleLeftMouseButton || bHandleMiddleMouseButton) && MouseDownPosition.IsSet())
+	else
 	{
-		TRange<double> LocalViewRange = GetViewRange();
-		FScrubRangeToScreen RangeToScreen(LocalViewRange, MyGeometry.Size);
-		DistanceDragged += FMath::Abs( MouseEvent.GetCursorDelta().X );
+		const FFrameTime MouseTime = ComputeFrameTimeFromMouse(InGeometry, InPointerEvent.GetScreenSpacePosition(), InRangeToScreen);
+		const FFrameTime MouseDownTime = ComputeFrameTimeFromMouse(InGeometry, MouseDownPosition.GetValue(), InRangeToScreen);
+		const FFrameNumber DiffFrame = MouseTime.FrameNumber - MouseDownTime.FrameNumber;
 
-		if ( MouseDragType == DRAG_NONE )
+		// Set the start playback range time?
+		if (MouseDragType == DRAG_PLAYBACK_START)
 		{
-			if ( DistanceDragged > 0.f /*FSlateApplication::Get().GetDragTriggerDistance()*/ )
+			if (InPointerEvent.IsShiftDown())
 			{
-				FFrameTime MouseDownFree = ComputeFrameTimeFromMouse(MyGeometry, MouseDownPosition.GetValue(), RangeToScreen, false);
-
-				const bool       bReadOnly          = Sequencer->IsReadOnly();
-				const FFrameRate TickResolution     = GetTickResolution();
-				const bool       bLockedPlayRange   = TimeSliderArgs.IsPlaybackRangeLocked.Get();
-				const float      MouseDownPixel     = RangeToScreen.InputToLocalX(MouseDownFree / TickResolution);
-				const bool       bHitScrubber       = GetHitTestScrubPixelMetrics(RangeToScreen).HandleRangePx.Contains(MouseDownPixel);
-
-				TRange<double>   SelectionRange   = TimeSliderArgs.SelectionRange.Get() / TickResolution;
-				TRange<double>   PlaybackRange    = TimeSliderArgs.PlaybackRange.Get()  / TickResolution;
-
-				// Disable selection range test if it's empty so that the playback range scrubbing gets priority
-				if (!SelectionRange.IsEmpty() && !bHitScrubber && HitTestRangeEnd(RangeToScreen, SelectionRange, MouseDownPixel) && bHandleMiddleMouseButton == false)
-				{
-					// selection range end scrubber
-					MouseDragType = DRAG_SELECTION_END;
-					TimeSliderArgs.OnSelectionRangeBeginDrag.ExecuteIfBound();
-				}
-				else if (!SelectionRange.IsEmpty() && !bHitScrubber && HitTestRangeStart(RangeToScreen, SelectionRange, MouseDownPixel) && bHandleMiddleMouseButton == false)
-				{
-					// selection range start scrubber
-					MouseDragType = DRAG_SELECTION_START;
-					TimeSliderArgs.OnSelectionRangeBeginDrag.ExecuteIfBound();
-				}
-				else if (!bLockedPlayRange && !bHitScrubber && HitTestRangeEnd(RangeToScreen, PlaybackRange, MouseDownPixel) && bHandleMiddleMouseButton == false)
-				{
-					// playback range end scrubber
-					MouseDragType = DRAG_PLAYBACK_END;
-					TimeSliderArgs.OnPlaybackRangeBeginDrag.ExecuteIfBound();
-				}
-				else if (!bLockedPlayRange && !bHitScrubber && HitTestRangeStart(RangeToScreen, PlaybackRange, MouseDownPixel) && bHandleMiddleMouseButton == false)
-				{
-					// playback range start scrubber
-					MouseDragType = DRAG_PLAYBACK_START;
-					TimeSliderArgs.OnPlaybackRangeBeginDrag.ExecuteIfBound();
-				}
-				else if (!bLockedMarkedFrames && !bHitScrubber && HitTestMark(MyGeometry, RangeToScreen, MouseDownPixel, bFromTimeSlider, &DragMarkIndex) && bHandleMiddleMouseButton == false)
-				{
-					MouseDragType = DRAG_MARK;
-					HandleMarkSelection(DragMarkIndex);
-
-					TimeSliderArgs.OnMarkBeginDrag.ExecuteIfBound();
-				}
-				else if (FSlateApplication::Get().GetModifierKeys().AreModifersDown(EModifierKey::Control) && bHandleMiddleMouseButton == false)
-				{
-					MouseDragType = DRAG_SETTING_RANGE;
-				}
-				else if (bMouseDownInRegion)
-				{
-					MouseDragType = DRAG_SCRUBBING_TIME;
-					TimeSliderArgs.OnBeginScrubberMovement.ExecuteIfBound();
-				}
+				SetPlaybackRangeStart(MouseDownPlaybackRange.GetLowerBoundValue() + DiffFrame);
+				SetPlaybackRangeEnd(MouseDownPlaybackRange.GetUpperBoundValue() + DiffFrame);
+			}
+			else
+			{
+				SetPlaybackRangeStart(MouseTime.FrameNumber);
 			}
 		}
-		else
+		// Set the end playback range time?
+		else if (MouseDragType == DRAG_PLAYBACK_END)
 		{
-			FFrameTime MouseTime = ComputeFrameTimeFromMouse(MyGeometry, MouseEvent.GetScreenSpacePosition(), RangeToScreen);
-			FFrameTime ScrubTime = ComputeScrubTimeFromMouse(MyGeometry, MouseEvent, RangeToScreen);
-			FFrameTime MouseDownTime = ComputeFrameTimeFromMouse(MyGeometry, MouseDownPosition.GetValue(), RangeToScreen);
-			FFrameTime LastMouseTime = ComputeFrameTimeFromMouse(MyGeometry, LastMousePosition.GetValue(), RangeToScreen);
-			FFrameNumber DiffFrame = MouseTime.FrameNumber - MouseDownTime.FrameNumber;
-			FFrameNumber LastDiffFrame = MouseTime.FrameNumber - LastMouseTime.FrameNumber;
-
-			LastMousePosition = MouseEvent.GetScreenSpacePosition();
-
-			// Set the start range time?
-			if (MouseDragType == DRAG_PLAYBACK_START)
+			if (InPointerEvent.IsShiftDown())
 			{
-				if (MouseEvent.IsShiftDown())
-				{
-					SetPlaybackRangeStart(MouseDownPlaybackRange.GetLowerBoundValue() + DiffFrame);
-					SetPlaybackRangeEnd(MouseDownPlaybackRange.GetUpperBoundValue() + DiffFrame);
-				}
-				else
-				{
-					SetPlaybackRangeStart(MouseTime.FrameNumber);
-				}
+				SetPlaybackRangeStart(MouseDownPlaybackRange.GetLowerBoundValue() + DiffFrame);
+				SetPlaybackRangeEnd(MouseDownPlaybackRange.GetUpperBoundValue() + DiffFrame);
 			}
-			// Set the end range time?
-			else if(MouseDragType == DRAG_PLAYBACK_END)
-			{
-				if (MouseEvent.IsShiftDown())
-				{
-					SetPlaybackRangeStart(MouseDownPlaybackRange.GetLowerBoundValue() + DiffFrame);
-					SetPlaybackRangeEnd(MouseDownPlaybackRange.GetUpperBoundValue() + DiffFrame);
-				}
-				else
-				{		
-					SetPlaybackRangeEnd(MouseTime.FrameNumber);
-				}
+			else
+			{		
+				SetPlaybackRangeEnd(MouseTime.FrameNumber);
 			}
-			else if (MouseDragType == DRAG_SELECTION_START)
+		}
+		// Set the start selection range time?
+		else if (MouseDragType == DRAG_SELECTION_START)
+		{
+			if (InPointerEvent.IsShiftDown())
 			{
-				if (MouseEvent.IsShiftDown())
-				{
-					SetSelectionRangeStart(MouseDownSelectionRange.GetLowerBoundValue() + DiffFrame);
-					SetSelectionRangeEnd(MouseDownSelectionRange.GetUpperBoundValue() + DiffFrame);
-				}
-				else
-				{
-					SetSelectionRangeStart(MouseTime.FrameNumber);
-				}
+				SetSelectionRangeStart(MouseDownSelectionRange.GetLowerBoundValue() + DiffFrame);
+				SetSelectionRangeEnd(MouseDownSelectionRange.GetUpperBoundValue() + DiffFrame);
 			}
-			// Set the end range time?
-			else if(MouseDragType == DRAG_SELECTION_END)
+			else
 			{
-				if (MouseEvent.IsShiftDown())
-				{
-					SetSelectionRangeStart(MouseDownSelectionRange.GetLowerBoundValue() + DiffFrame);
-					SetSelectionRangeEnd(MouseDownSelectionRange.GetUpperBoundValue() + DiffFrame);
-				}
-				else 
-				{
-					SetSelectionRangeEnd(MouseTime.FrameNumber);
-				}
+				SetSelectionRangeStart(MouseTime.FrameNumber);
 			}
-			else if (MouseDragType == DRAG_MARK)
+		}
+		// Set the end selection range time?
+		else if (MouseDragType == DRAG_SELECTION_END)
+		{
+			if (InPointerEvent.IsShiftDown())
 			{
-				Sequencer->TransformSelection(LastDiffFrame, 1.f);
+				SetSelectionRangeStart(MouseDownSelectionRange.GetLowerBoundValue() + DiffFrame);
+				SetSelectionRangeEnd(MouseDownSelectionRange.GetUpperBoundValue() + DiffFrame);
 			}
-			else if (MouseDragType == DRAG_SCRUBBING_TIME)
+			else 
 			{
-				// Delegate responsibility for clamping to the current viewrange to the client
-				CommitScrubPosition( ScrubTime, /*bIsScrubbing=*/true, /*bEvaluate*/ !bHandleMiddleMouseButton); //if middle mouse button down we don't evaluate on the time change
+				SetSelectionRangeEnd(MouseTime.FrameNumber);
 			}
 		}
+		else if (MouseDragType == DRAG_MARK)
+		{
+			const FVector2D PreviousMousePosition = LastMousePosition.Get(InPointerEvent.GetScreenSpacePosition());
+			const FFrameTime LastMouseTime = ComputeFrameTimeFromMouse(InGeometry, PreviousMousePosition, InRangeToScreen);
+			const FFrameNumber LastDiffFrame = MouseTime.FrameNumber - LastMouseTime.FrameNumber;
+			HandleDragMark(LastDiffFrame);
+		}
+		else if (MouseDragType == DRAG_SCRUBBING_TIME)
+		{
+			const FFrameTime ScrubTime = ComputeScrubTimeFromMouse(InGeometry, InPointerEvent, InRangeToScreen);
+			// Delegate responsibility for clamping to the current viewrange to the client
+			// If middle mouse button down we don't evaluate on the time change
+			CommitScrubPosition(ScrubTime, /*bIsScrubbing=*/true, /*bEvaluate=*/!bHandleMiddleMouseButton);
+		}
+
+		LastMousePosition = InPointerEvent.GetScreenSpacePosition();
 	}
-	else if (bFromTimeSlider && !bLockedMarkedFrames && DragMarkIndex == INDEX_NONE)
+
+	return FReply::Unhandled();
+}
+
+FReply FSequencerTimeSliderController::OnMouseMove_RightMouseDown(const FGeometry& InGeometry
+	, const FPointerEvent& InPointerEvent
+	, const USequencerSettings& InSequencerSettings)
+{
+	DistanceDragged += FMath::Abs(InPointerEvent.GetCursorDelta().X);
+
+	if (!bPanning)
 	{
-		// Update hover state of marked frames.
-		TRange<double> LocalViewRange = GetViewRange();
-		FScrubRangeToScreen RangeToScreen(LocalViewRange, MyGeometry.Size);
-		const FVector2D LocalMousePostion = MyGeometry.AbsoluteToLocal(MouseEvent.GetScreenSpacePosition());
-		HitTestMark(MyGeometry, RangeToScreen, LocalMousePostion.X, bFromTimeSlider, &HoverMarkIndex);
+		if (DistanceDragged > 0.f/*FSlateApplication::Get().GetDragTriggerDistance()*/)
+		{
+			bPanning = true;
+		}
 	}
-
-	if ( DistanceDragged != 0.f && (bHandleLeftMouseButton || bHandleRightMouseButton) )
+	else if (InPointerEvent.IsShiftDown() && InPointerEvent.IsAltDown())
 	{
-		return FReply::Handled().CaptureMouse(WidgetOwner.AsShared());
+		const FVector2D PointerPosition = InPointerEvent.GetScreenSpacePosition();
+		const FFrameRate TickResolution = GetTickResolution();
+
+		float MouseFractionX = InGeometry.AbsoluteToLocal(PointerPosition).X / InGeometry.GetLocalSize().X;
+
+		// If zooming on the current time, adjust mouse fractionX
+		if (InSequencerSettings.GetZoomPosition() == ESequencerZoomPosition::SZP_CurrentTime)
+		{
+			const double ScrubPosition = TimeSliderArgs.ScrubPosition.Get() / TickResolution;
+			if (GetViewRange().Contains(ScrubPosition))
+			{
+				FScrubRangeToScreen RangeToScreen(GetViewRange(), InGeometry.Size);
+				const float TimePosition = RangeToScreen.InputToLocalX(ScrubPosition);
+				MouseFractionX = TimePosition / InGeometry.GetLocalSize().X;
+			}
+		}
+
+		const float ZoomDelta = -0.01f * InPointerEvent.GetCursorDelta().X;
+		ZoomByDelta(ZoomDelta, MouseFractionX);
 	}
+	else
+	{
+		const TRange<double> LocalViewRange = GetViewRange();
+		const double LocalViewRangeMin = LocalViewRange.GetLowerBoundValue();
+		const double LocalViewRangeMax = LocalViewRange.GetUpperBoundValue();
+
+		const FScrubRangeToScreen ScaleInfo(LocalViewRange, InGeometry.Size);
+		const FVector2D ScreenDelta = InPointerEvent.GetCursorDelta();
+		const FVector2D InputDelta = FVector2D(ScreenDelta.X / ScaleInfo.PixelsPerInput, 0.);
+
+		double NewViewOutputMin = LocalViewRangeMin - InputDelta.X;
+		double NewViewOutputMax = LocalViewRangeMax - InputDelta.X;
 
+		ClampViewRange(NewViewOutputMin, NewViewOutputMax);
+		SetViewRange(NewViewOutputMin, NewViewOutputMax, EViewRangeInterpolation::Immediate);
+	}
 
 	return FReply::Handled();
 }
 
-
 void FSequencerTimeSliderController::CommitScrubPosition( FFrameTime NewValue, bool bIsScrubbing, bool bEvaluate)
 {
 	using namespace UE::Sequencer;
@@ -1505,44 +1668,41 @@ void FSequencerTimeSliderController::CommitScrubPosition( FFrameTime NewValue, b
 	TimeSliderArgs.OnScrubPositionChanged.ExecuteIfBound( NewValue, bIsScrubbing, bEvaluate);
 }
 
-FReply FSequencerTimeSliderController::OnMouseWheel( SWidget& WidgetOwner, const FGeometry& MyGeometry, const FPointerEvent& MouseEvent )
+FReply FSequencerTimeSliderController::OnMouseWheel(SWidget& OwnerWidget, const FGeometry& InGeometry, const FPointerEvent& InPointerEvent)
 {
-	TOptional<TRange<float>> NewTargetRange;
-
-	if ( TimeSliderArgs.AllowZoom && MouseEvent.IsControlDown() )
+	if (TimeSliderArgs.AllowZoom && InPointerEvent.IsControlDown())
 	{
-		float MouseFractionX = MyGeometry.AbsoluteToLocal(MouseEvent.GetScreenSpacePosition()).X / MyGeometry.GetLocalSize().X;
+		float MouseFractionX = InGeometry.AbsoluteToLocal(InPointerEvent.GetScreenSpacePosition()).X / InGeometry.GetLocalSize().X;
 
-		TSharedPtr<FSequencer> Sequencer = WeakSequencer.Pin();
-			
 		// If zooming on the current time, adjust mouse fractionX
+		const TSharedPtr<FSequencer> Sequencer = WeakSequencer.Pin();
 		if (Sequencer.IsValid() && Sequencer->GetSequencerSettings()->GetZoomPosition() == ESequencerZoomPosition::SZP_CurrentTime)
 		{
 			const double ScrubPosition = TimeSliderArgs.ScrubPosition.Get() / GetTickResolution();
 			if (GetViewRange().Contains(ScrubPosition))
 			{
-				FScrubRangeToScreen RangeToScreen(GetViewRange(), MyGeometry.Size);
+				FScrubRangeToScreen RangeToScreen(GetViewRange(), InGeometry.Size);
 				float TimePosition = RangeToScreen.InputToLocalX(ScrubPosition);
-				MouseFractionX = TimePosition / MyGeometry.GetLocalSize().X;
+				MouseFractionX = TimePosition / InGeometry.GetLocalSize().X;
 			}
 		}
 
-		const float ZoomDelta = -0.2f * MouseEvent.GetWheelDelta();
+		const float ZoomDelta = -0.2f * InPointerEvent.GetWheelDelta();
 		if (ZoomByDelta(ZoomDelta, MouseFractionX))
 		{
 			return FReply::Handled();
 		}
 	}
-	else if (MouseEvent.IsShiftDown())
+	else if (InPointerEvent.IsShiftDown())
 	{
-		PanByDelta(-MouseEvent.GetWheelDelta());
+		PanByDelta(-InPointerEvent.GetWheelDelta());
 		return FReply::Handled();
 	}
 	
 	return FReply::Unhandled();
 }
 
-FCursorReply FSequencerTimeSliderController::OnCursorQuery( TSharedRef<const SWidget> WidgetOwner, const FGeometry& MyGeometry, const FPointerEvent& CursorEvent ) const
+FCursorReply FSequencerTimeSliderController::OnCursorQuery( const SWidget& OwnerWidget, const FGeometry& MyGeometry, const FPointerEvent& CursorEvent ) const
 {
 	TSharedPtr<FSequencer> Sequencer = WeakSequencer.Pin();
 	if (!Sequencer.IsValid())
@@ -1639,13 +1799,13 @@ int32 FSequencerTimeSliderController::OnPaintViewArea( const FGeometry& Allotted
 		const USequencerSettings* const SequencerSettings = Sequencer->GetSequencerSettings();
 		if (SequencerSettings && SequencerSettings->ShouldDrawHorizontalTrackSeparatorLines())
 		{
-			DrawHorizontalTrackSeparatorLines(OutDrawElements, LocalViewRange, RangeToScreen, DrawTickArgs);
+			DrawHorizontalTrackSeparatorLines(OutDrawElements, DrawTickArgs);
 		}
 	}
 
 	if (Args.bDisplayMarkedFrames)
 	{
-		LayerId = DrawMarkedFrames(AllottedGeometry, RangeToScreen, OutDrawElements, LayerId, DrawEffects, FWidgetStyle(), false);
+		LayerId = DrawMarkedFrames(AllottedGeometry, RangeToScreen, OutDrawElements, LayerId, DrawEffects, false);
 	}
 
 	LayerId = DrawScalingAnchors(AllottedGeometry, RangeToScreen, OutDrawElements, LayerId, DrawEffects);
@@ -1844,36 +2004,37 @@ void FSequencerTimeSliderController::SetStoppedPosition(FFrameTime InTime)
 
 void FSequencerTimeSliderController::SetPlaybackStatus(ETimeSliderPlaybackStatus InStatus)
 {
-	using namespace UE::Sequencer;
-	if (!WeakSequencer.IsValid())
+	const TSharedPtr<FSequencer> Sequencer = WeakSequencer.Pin();
+	if (!Sequencer.IsValid())
 	{
 		return;
 	}
+
 	switch (InStatus)
 	{
 		case ETimeSliderPlaybackStatus::Jumping:
 		{
-			WeakSequencer.Pin()->SetPlaybackStatus(EMovieScenePlayerStatus::Jumping);
+			Sequencer->SetPlaybackStatus(EMovieScenePlayerStatus::Jumping);
 			break;
 		}
 		case ETimeSliderPlaybackStatus::Paused:
 		{
-			WeakSequencer.Pin()->SetPlaybackStatus(EMovieScenePlayerStatus::Paused);
+			Sequencer->SetPlaybackStatus(EMovieScenePlayerStatus::Paused);
 			break;
 		}
 		case ETimeSliderPlaybackStatus::Playing:
 		{
-			WeakSequencer.Pin()->SetPlaybackStatus(EMovieScenePlayerStatus::Playing);
+			Sequencer->SetPlaybackStatus(EMovieScenePlayerStatus::Playing);
 			break;
 		}
 		case ETimeSliderPlaybackStatus::Scrubbing:
 		{
-			WeakSequencer.Pin()->SetPlaybackStatus(EMovieScenePlayerStatus::Scrubbing);
+			Sequencer->SetPlaybackStatus(EMovieScenePlayerStatus::Scrubbing);
 			break;
 		}
 		case ETimeSliderPlaybackStatus::Stepping:
 		{
-			WeakSequencer.Pin()->SetPlaybackStatus(EMovieScenePlayerStatus::Stepping);
+			Sequencer->SetPlaybackStatus(EMovieScenePlayerStatus::Stepping);
 			break;
 		}
 		default:
@@ -1881,7 +2042,7 @@ void FSequencerTimeSliderController::SetPlaybackStatus(ETimeSliderPlaybackStatus
 		{
 			//Call Pause() instead of SetPlaybackStatus(EMovieScenePlayerStatus::Stopped)
 			//so the correct delegates/evals can happen so all systems like audio are probably stopped (Pause set's Stopped status).
-			WeakSequencer.Pin()->Pause();
+			Sequencer->Pause();
 			break;
 		}
 	}
@@ -1889,45 +2050,33 @@ void FSequencerTimeSliderController::SetPlaybackStatus(ETimeSliderPlaybackStatus
 
 ETimeSliderPlaybackStatus FSequencerTimeSliderController::GetPlaybackStatus() const
 {
-	using namespace UE::Sequencer;
-	if (!WeakSequencer.IsValid())
+	const TSharedPtr<FSequencer> Sequencer = WeakSequencer.Pin();
+	if (!Sequencer.IsValid())
 	{
 		return ETimeSliderPlaybackStatus::Stopped;
 	}
-	EMovieScenePlayerStatus::Type Status = WeakSequencer.Pin()->GetPlaybackStatus();
+
+	const EMovieScenePlayerStatus::Type Status = Sequencer->GetPlaybackStatus();
 	switch (Status)
 	{
 		case EMovieScenePlayerStatus::Jumping:
-		{
 			return ETimeSliderPlaybackStatus::Jumping;
-		}
 		case EMovieScenePlayerStatus::Paused:
-		{
 			return ETimeSliderPlaybackStatus::Paused;
-		}
 		case EMovieScenePlayerStatus::Playing:
-		{
 			return ETimeSliderPlaybackStatus::Playing;
-		}
 		case EMovieScenePlayerStatus::Scrubbing:
-		{
 			return ETimeSliderPlaybackStatus::Scrubbing;
-		}
 		case EMovieScenePlayerStatus::Stepping:
-		{
 			return ETimeSliderPlaybackStatus::Stepping;
-		}
 		case EMovieScenePlayerStatus::Stopped:
-		{
 			return ETimeSliderPlaybackStatus::Stopped;
-		}
-
 	}
-	return ETimeSliderPlaybackStatus::Stopped;
 
+	return ETimeSliderPlaybackStatus::Stopped;
 }
 
-TSharedRef<SWidget> FSequencerTimeSliderController::OpenSetPlaybackRangeMenu(const FGeometry& MyGeometry, const FPointerEvent& MouseEvent)
+TSharedPtr<SWidget> FSequencerTimeSliderController::OpenContextMenu(const FGeometry& InGeometry, const FPointerEvent& InPointerEvent)
 {
 	using namespace UE::Sequencer;
 	
@@ -2308,9 +2457,13 @@ TSharedRef<SWidget> FSequencerTimeSliderController::OpenSetPlaybackRangeMenu(con
 	
 	USequencerTimeSliderControllerMenuContext* TimeSliderContext = NewObject<USequencerTimeSliderControllerMenuContext>();
 	TimeSliderContext->WeakTimeSliderController = AsWeak();
-	TimeSliderContext->Geometry = MyGeometry;
-	TimeSliderContext->PointerEvent = MouseEvent;
-	TimeSliderContext->FrameNumber = ComputeFrameTimeFromMouse(MyGeometry, MouseEvent.GetScreenSpacePosition(), FScrubRangeToScreen(GetViewRange(), MyGeometry.Size)).FrameNumber;
+	TimeSliderContext->Geometry = InGeometry;
+	TimeSliderContext->PointerEvent = InPointerEvent;
+	TimeSliderContext->FrameNumber = ComputeFrameTimeFromMouse(
+		InGeometry,
+		InPointerEvent.GetScreenSpacePosition(),
+		FScrubRangeToScreen(GetViewRange(), InGeometry.Size)
+	).FrameNumber;
 
 	USequencerToolMenuContext* SequencerContext = NewObject<USequencerToolMenuContext>();
 	SequencerContext->WeakSequencer = WeakSequencer;
@@ -2663,7 +2816,7 @@ void FSequencerTimeSliderController::SetSelectionRangeEnd(FFrameNumber NewEnd)
 	}
 }
 
-void FSequencerTimeSliderController::HandleMarkSelection(int32 InMarkIndex)
+void FSequencerTimeSliderController::HandleDragMark_NoDragType(const int32 InMarkIndex)
 {
 	using namespace UE::Sequencer;
 
@@ -2675,7 +2828,7 @@ void FSequencerTimeSliderController::HandleMarkSelection(int32 InMarkIndex)
 
 	FSequencerSelection& SequencerSelection = Sequencer->GetSelection();
 
-	FModifierKeysState ModifierKeys = FSlateApplication::Get().GetModifierKeys();
+	const FModifierKeysState ModifierKeys = FSlateApplication::Get().GetModifierKeys();
 	const bool bToggleSelection = ModifierKeys.AreModifersDown(EModifierKey::Control);
 	const bool bAddToSelection = ModifierKeys.AreModifersDown(EModifierKey::Shift);
 
@@ -2697,12 +2850,41 @@ void FSequencerTimeSliderController::HandleMarkSelection(int32 InMarkIndex)
 	}
 }
 
-void FSequencerTimeSliderController::AddMarkAtFrame(FFrameNumber FrameNumber)
+void FSequencerTimeSliderController::HandleDragMark(const FFrameNumber& InLastDiffFrame)
 {
-	TimeSliderArgs.OnAddMarkedFrame.ExecuteIfBound(FrameNumber);
+	if (const TSharedPtr<FSequencer> Sequencer = WeakSequencer.Pin())
+	{
+		Sequencer->TransformSelection(InLastDiffFrame, 1.f);
+	}
 }
 
-void FSequencerTimeSliderController::DeleteMarkAtIndex(int32 InMarkIndex)
+void FSequencerTimeSliderController::ClearMarkSelection()
+{
+	if (const TSharedPtr<FSequencer> Sequencer = WeakSequencer.Pin())
+	{
+		Sequencer->GetSelection().MarkedFrames.Empty();
+	}
+}
+
+TOptional<FFrameNumber> FSequencerTimeSliderController::GetMarkFrameNumber(const int32 InMarkIndex) const
+{
+	const TArray<FMovieSceneMarkedFrame>& MarkedFrames = TimeSliderArgs.MarkedFrames.Get();
+	return MarkedFrames.IsValidIndex(InMarkIndex)
+		? TOptional<FFrameNumber>(MarkedFrames[InMarkIndex].FrameNumber)
+		: TOptional<FFrameNumber>();
+}
+
+void FSequencerTimeSliderController::AddMarkAtFrame(const FFrameNumber InFrameNumber)
+{
+	TimeSliderArgs.OnAddMarkedFrame.ExecuteIfBound(InFrameNumber);
+}
+
+void FSequencerTimeSliderController::SetMarkAtFrame(const int32 InMarkIndex, const FFrameNumber InFrameNumber)
+{
+	TimeSliderArgs.OnSetMarkedFrame.ExecuteIfBound(InMarkIndex, InFrameNumber);
+}
+
+void FSequencerTimeSliderController::DeleteMarkAtIndex(const int32 InMarkIndex)
 {
 	TimeSliderArgs.OnDeleteMarkedFrame.ExecuteIfBound(InMarkIndex);
 }
@@ -2711,6 +2893,18 @@ void FSequencerTimeSliderController::DeleteAllMarks()
 {
 	TimeSliderArgs.OnDeleteAllMarkedFrames.ExecuteIfBound();
 }
-	
+
+TSharedPtr<ISequencer> FSequencerTimeSliderController::GetSequencer() const
+{
+	return WeakSequencer.Pin();
+}
+
+void FSequencerTimeSliderController::ResetMouseInput()
+{
+	bPanning = false;
+	bMouseDownInRegion = false;
+	MouseDownPosition.Reset();
+	LastMousePosition.Reset();
+}
 
 #undef LOCTEXT_NAMESPACE
diff --git a/Engine/Source/Editor/Sequencer/Private/SimpleView/Menus/SequencerToolbarExtension.h b/Engine/Source/Editor/Sequencer/Private/SimpleView/Menus/SequencerToolbarExtension.h
new file mode 100644
--- /dev/null
+++ b/Engine/Source/Editor/Sequencer/Private/SimpleView/Menus/SequencerToolbarExtension.h
@@ -0,0 +1,46 @@
+﻿// Copyright Epic Games, Inc. All Rights Reserved.
+
+#pragma once
+
+#include "Templates/SharedPointer.h"
+
+class FSequencer;
+class FUICommandInfo;
+class UToolableTimelineMenuContext;
+struct FFrameTime;
+struct FToolMenuSection;
+
+namespace UE::Sequencer::SimpleView
+{
+
+class FSimpleViewTimeline;
+
+/** Extension for Sequencer toolbar extensions */
+class FSequencerToolbarExtension
+{
+public:
+	static const FName SequencerToolbarMenuName;
+	static const FName SequencerToolbarOwner;
+	static const FName SimpleViewSectionName;
+
+	static void AddExtension();
+	static void RemoveExtension();
+
+protected:
+	static void BuildSimpleViewSection(FToolMenuSection& InSection);
+
+	static void CreateKeyTranslateSection(FToolMenuSection& InSection, const UToolableTimelineMenuContext& InContext);
+	static void CreateKeyScaleSection(FToolMenuSection& InSection, const UToolableTimelineMenuContext& InContext);
+
+	static void ExecuteSimpleViewAction(const TWeakPtr<FSimpleViewTimeline> InWeakTimeline
+		, const TSharedRef<FUICommandInfo> InAction);
+	static bool CanExecuteSimpleViewAction(const TWeakPtr<FSimpleViewTimeline> InWeakTimeline
+		, const TArray<TSharedPtr<FUICommandInfo>> InActions);
+
+	static TSharedPtr<FSequencer> GetSequencerFromTimeline(const TWeakPtr<FSimpleViewTimeline>& InWeakTimeline);
+
+	static void TransformSelection(const TWeakPtr<FSimpleViewTimeline> InWeakTimeline
+		, const FFrameTime& InDeltaTime, const float InScale);
+};
+
+} // namespace UE::Sequencer::SimpleView
diff --git a/Engine/Source/Editor/Sequencer/Private/ToolableTimeline/Caches/ToolableTimelineChannelCache.cpp b/Engine/Source/Editor/Sequencer/Private/ToolableTimeline/Caches/ToolableTimelineChannelCache.cpp
new file mode 100644
--- /dev/null
+++ b/Engine/Source/Editor/Sequencer/Private/ToolableTimeline/Caches/ToolableTimelineChannelCache.cpp
@@ -0,0 +1,256 @@
+﻿// Copyright Epic Games, Inc. All Rights Reserved.
+
+#include "ToolableTimelineChannelCache.h"
+#include "IKeyArea.h"
+#include "MVVM/Extensions/ITopLevelChannelHolderExtension.h"
+#include "MVVM/Selection/Selection.h"
+#include "Sequencer.h"
+#include "ToolableTimeline/ToolableTimeline.h"
+
+namespace UE::Sequencer::ToolableTimeline
+{
+
+FToolableTimelineChannelCache::FToolableTimelineChannelCache(const FToolableTimeline& InTimeline)
+{
+	const FSequencer& Sequencer = *InTimeline.GetSequencerChecked();
+
+	switch (InTimeline.GetTimelineSettings().Settings.KeyDisplay)
+	{
+	default:
+	case ETimelineKeyDisplay::SelectedAndPinned:
+		WeakChannelModels.Append(GetChannelsFromSelectedOutlinerItems(Sequencer));
+		WeakChannelModels.Append(GetChannelsFromPinnedOutlinerItems(Sequencer));
+		break;
+
+	case ETimelineKeyDisplay::Selected:
+		WeakChannelModels.Append(GetChannelsFromSelectedOutlinerItems(Sequencer));
+		break;
+
+	case ETimelineKeyDisplay::All:
+		WeakChannelModels.Append(GetChannelsWithKeys(Sequencer));
+		break;
+	}
+
+	ApplyExtendedFilters(InTimeline);
+}
+
+const TSet<TWeakViewModelPtr<FChannelModel>>& FToolableTimelineChannelCache::GetChannelModels() const
+{
+	return WeakChannelModels;
+}
+
+bool FToolableTimelineChannelCache::DoesChannelModelHaveKeys(const TSharedPtr<FChannelModel>& InChannelModel
+	, const TRange<FFrameNumber>& InFrameRange)
+{
+	if (!InChannelModel.IsValid())
+	{
+		return false;
+	}
+
+	const TSharedPtr<IKeyArea> KeyArea = InChannelModel->GetKeyArea();
+	if (!KeyArea.IsValid())
+	{
+		return false;
+	}
+
+	TArray<FKeyHandle> KeyHandles;
+	KeyArea->GetKeyHandles(KeyHandles, InFrameRange);
+
+	return !KeyHandles.IsEmpty();
+}
+
+bool FToolableTimelineChannelCache::GetChannelsWithDescendantKeys(const FViewModelPtr& InModel
+	, TSet<TWeakViewModelPtr<FChannelModel>>& OutWeakChannelModels)
+{
+	if (!InModel.IsValid())
+	{
+		return false;
+	}
+
+	bool bOutReturn = false;
+
+	for (const TViewModelPtr<FChannelModel>& ChildChannelModel : InModel->GetDescendantsOfType<FChannelModel>(/*bIncludeThis=*/true))
+	{
+		if (DoesChannelModelHaveKeys(ChildChannelModel))
+		{
+			OutWeakChannelModels.Add(ChildChannelModel);
+			bOutReturn = true;
+		}
+	}
+
+	if (const TViewModelPtr<ITopLevelChannelHolderExtension> ChannelHolder = InModel.ImplicitCast())
+	{
+		for (const TViewModelPtr<FChannelGroupModel>& ChannelGroupModel
+			: ChannelHolder->GetTopLevelChannels().IterateSubList<FChannelGroupModel>())
+		{
+			for (const TWeakViewModelPtr<FChannelModel>& WeakChannelModel : ChannelGroupModel->GetChannels())
+			{
+				const TViewModelPtr<FChannelModel> ChannelModel = WeakChannelModel.Pin();
+				if (ChannelModel.IsValid() && DoesChannelModelHaveKeys(ChannelModel))
+				{
+					OutWeakChannelModels.Add(ChannelModel);
+					bOutReturn = true;
+				}
+			}
+		}
+	}
+
+	return bOutReturn;
+}
+
+TSet<TWeakViewModelPtr<FChannelModel>> FToolableTimelineChannelCache::GetChannelsWithKeys(const ISequencer& InSequencer)
+{
+	TSet<TWeakViewModelPtr<FChannelModel>> OutWeakChannelModels;
+
+	const TSharedPtr<FSequencerEditorViewModel> RootViewModel = InSequencer.GetViewModel();
+	if (!RootViewModel.IsValid())
+	{
+		return OutWeakChannelModels;
+	}
+
+	const TViewModelPtr<FSequenceModel> RootSequenceModel = RootViewModel->GetRootSequenceModel();
+	if (!RootSequenceModel.IsValid())
+	{
+		return OutWeakChannelModels;
+	}
+
+	GetChannelsWithDescendantKeys(RootSequenceModel, OutWeakChannelModels);
+
+	return OutWeakChannelModels;
+}
+
+void FToolableTimelineChannelCache::GatherChannelsFromViewModel(const TViewModelPtr<FViewModel>& InViewModel
+	, TSet<TWeakViewModelPtr<FChannelModel>>& OutWeakChannelModels)
+{
+	if (!InViewModel.IsValid())
+	{
+		return;
+	}
+
+	if (FChannelModel* const ChannelModel = InViewModel->CastThis<FChannelModel>())
+	{
+		if (ChannelModel->IsAnimated())
+		{
+			OutWeakChannelModels.Add(InViewModel->CastThisShared<FChannelModel>());
+		}
+	}
+
+	if (FChannelGroupModel* const ChannelGroup = InViewModel->CastThis<FChannelGroupModel>())
+	{
+		for (const TWeakViewModelPtr<FChannelModel>& WeakChannel : ChannelGroup->GetChannels())
+		{
+			const TViewModelPtr<FChannelModel> Channel = WeakChannel.Pin();
+			if (Channel.IsValid() && Channel->IsAnimated())
+			{
+				OutWeakChannelModels.Add(WeakChannel);
+			}
+		}
+	}
+
+	for (FViewModelListIterator ChildIt = InViewModel->GetChildren(EViewModelListType::Outliner); ChildIt; ++ChildIt)
+	{
+		if (const TViewModelPtr<IOutlinerExtension> ChildOutlinerItem = ChildIt->CastThisShared<IOutlinerExtension>())
+		{
+			GatherChannelsFromViewModel(ChildOutlinerItem, OutWeakChannelModels);
+		}
+	}
+}
+
+TSet<TWeakViewModelPtr<FChannelModel>> FToolableTimelineChannelCache::GetChannelsFromSelectedOutlinerItems(const ISequencer& InSequencer)
+{
+	TSet<TWeakViewModelPtr<FChannelModel>> OutWeakChannelModels;
+
+	const TSharedPtr<FSequencerEditorViewModel> RootViewModel = InSequencer.GetViewModel();
+	if (!RootViewModel.IsValid())
+	{
+		return OutWeakChannelModels;
+	}
+
+	const TSharedPtr<FSequencerSelection> Selection = RootViewModel->GetSelection();
+	if (!Selection.IsValid())
+	{
+		return OutWeakChannelModels;
+	}
+
+	for (const TViewModelPtr<IOutlinerExtension>& SelectedOutlinerItem : Selection->Outliner)
+	{
+		GatherChannelsFromViewModel(SelectedOutlinerItem, OutWeakChannelModels);
+	}
+
+	return OutWeakChannelModels;
+}
+
+void FToolableTimelineChannelCache::GatherChannelsFromPinnedOutlinerItem(const TViewModelPtr<FViewModel>& InViewModel
+	, TSet<TWeakViewModelPtr<FChannelModel>>& OutWeakChannelModels)
+{
+	if (!InViewModel.IsValid())
+	{
+		return;
+	}
+
+	if (const IPinnableExtension* const PinnableExtension = InViewModel->CastThis<IPinnableExtension>())
+	{
+		if (PinnableExtension && PinnableExtension->IsPinned())
+		{
+			GatherChannelsFromViewModel(InViewModel, OutWeakChannelModels);
+			return;
+		}
+	}
+
+	for (FViewModelListIterator ChildIt = InViewModel->GetChildren(EViewModelListType::Outliner); ChildIt; ++ChildIt)
+	{
+		GatherChannelsFromPinnedOutlinerItem(*ChildIt, OutWeakChannelModels);
+	}
+}
+
+TSet<TWeakViewModelPtr<FChannelModel>> FToolableTimelineChannelCache::GetChannelsFromPinnedOutlinerItems(const ISequencer& InSequencer)
+{
+	TSet<TWeakViewModelPtr<FChannelModel>> OutWeakChannelModels;
+
+	const TSharedPtr<FSequencerEditorViewModel> RootViewModel = InSequencer.GetViewModel();
+	if (!RootViewModel.IsValid())
+	{
+		return OutWeakChannelModels;
+	}
+
+	const TViewModelPtr<FSequenceModel> RootSequenceModel = RootViewModel->GetRootSequenceModel();
+	if (!RootSequenceModel.IsValid())
+	{
+		return OutWeakChannelModels;
+	}
+
+	GatherChannelsFromPinnedOutlinerItem(RootSequenceModel, OutWeakChannelModels);
+
+	return OutWeakChannelModels;
+}
+
+void FToolableTimelineChannelCache::ApplyExtendedFilters(const FToolableTimeline& InTimeline)
+{
+	const TSharedPtr<FSequencer> Sequencer = InTimeline.GetSequencer();
+	if (!Sequencer.IsValid())
+	{
+		return;
+	}
+
+	if (InTimeline.GetChannelFilters().IsEmpty())
+	{
+		return;
+	}
+
+	for (auto It = WeakChannelModels.CreateIterator(); It; ++It)
+	{
+		const TViewModelPtr<FChannelModel> ChannelModel = It->Pin();
+		if (!ChannelModel.IsValid())
+		{
+			It.RemoveCurrent();
+			continue;
+		}
+
+		if (!InTimeline.PassesChannelFilters(*ChannelModel))
+		{
+			It.RemoveCurrent();
+		}
+	}
+}
+
+} // namespace UE::Sequencer::ToolableTimeline
diff --git a/Engine/Source/Editor/Sequencer/Private/ToolableTimeline/Caches/ToolableTimelineChannelCache.h b/Engine/Source/Editor/Sequencer/Private/ToolableTimeline/Caches/ToolableTimelineChannelCache.h
new file mode 100644
--- /dev/null
+++ b/Engine/Source/Editor/Sequencer/Private/ToolableTimeline/Caches/ToolableTimelineChannelCache.h
@@ -0,0 +1,69 @@
+﻿// Copyright Epic Games, Inc. All Rights Reserved.
+
+#pragma once
+
+#include "MVVM/ViewModelPtr.h"
+#include "MVVM/ViewModels/ChannelModel.h"
+#include "ToolableTimelineChannelCache.generated.h"
+
+class ISequencer;
+
+UENUM()
+enum class ETimelineKeyDisplay : uint8
+{
+	/** Show keys for selected and pinned objects on the timeline */
+	SelectedAndPinned,
+	/** Show keys for selected objects on the timeline */
+	Selected,
+	/** Show all keys for all objects on the timeline */
+	All
+};
+
+namespace UE::Sequencer
+{
+	class FChannelModel;
+}
+
+namespace UE::Sequencer::ToolableTimeline
+{
+
+class FToolableTimeline;
+
+class FToolableTimelineChannelCache
+{
+public:
+	FToolableTimelineChannelCache(const FToolableTimeline& InTimeline);
+
+	const TSet<TWeakViewModelPtr<FChannelModel>>& GetChannelModels() const;
+
+	/** @return True if the specified channel model has any keys */
+	static bool DoesChannelModelHaveKeys(const TSharedPtr<FChannelModel>& InChannelModel
+		, const TRange<FFrameNumber>& InFrameRange = TRange<FFrameNumber>::All());
+
+protected:
+	/** @return Set of channel models that contain keys */
+	static TSet<TWeakViewModelPtr<FChannelModel>> GetChannelsWithKeys(const ISequencer& InSequencer);
+
+	/** @return Set of channel models that contain selected outliner item keys */
+	static TSet<TWeakViewModelPtr<FChannelModel>> GetChannelsFromSelectedOutlinerItems(const ISequencer& InSequencer);
+
+	/** @return Set of pinned channel models that contain keys */
+	static TSet<TWeakViewModelPtr<FChannelModel>> GetChannelsFromPinnedOutlinerItems(const ISequencer& InSequencer);
+
+private:
+	static void GatherChannelsFromViewModel(const TViewModelPtr<FViewModel>& InViewModel
+		, TSet<TWeakViewModelPtr<FChannelModel>>& OutWeakChannelModels);
+	
+	static bool GetChannelsWithDescendantKeys(const FViewModelPtr& InModel
+		, TSet<TWeakViewModelPtr<FChannelModel>>& OutWeakChannelModels);
+
+	static void GatherChannelsFromPinnedOutlinerItem(const TViewModelPtr<FViewModel>& InViewModel
+		, TSet<TWeakViewModelPtr<FChannelModel>>& OutWeakChannelModels);
+
+	void ApplyExtendedFilters(const FToolableTimeline& InTimeline);
+
+	/** Cached channel models holding keys that are being processed */
+	TSet<TWeakViewModelPtr<FChannelModel>> WeakChannelModels;
+};
+
+} // namespace UE::Sequencer::ToolableTimeline
diff --git a/Engine/Source/Editor/Sequencer/Private/ToolableTimeline/DragOperations/ToolableTimelineScrubDragOperation.cpp b/Engine/Source/Editor/Sequencer/Private/ToolableTimeline/DragOperations/ToolableTimelineScrubDragOperation.cpp
new file mode 100644
--- /dev/null
+++ b/Engine/Source/Editor/Sequencer/Private/ToolableTimeline/DragOperations/ToolableTimelineScrubDragOperation.cpp
@@ -0,0 +1,32 @@
+﻿// Copyright Epic Games, Inc. All Rights Reserved.
+
+#include "ToolableTimelineScrubDragOperation.h"
+#include "ToolableTimeline/MouseInputData.h"
+
+namespace UE::Sequencer::ToolableTimeline
+{
+
+FToolableTimelineScrubDragOperation::FToolableTimelineScrubDragOperation(const FMouseInputData& InMouseInput
+	, const TRange<double>& InViewRange)
+	: FToolableTimelineDragOperation(InMouseInput, InViewRange)
+{
+}
+
+void FToolableTimelineScrubDragOperation::UpdateDrag(const FMouseInputData& InMouseInput)
+{
+	FToolableTimelineDragOperation::UpdateDrag(InMouseInput);
+
+	const TSharedRef<FToolableTimeSliderController> TimeSliderController = InMouseInput.Timeline->GetTimeSliderController();
+	const bool bSnapEnabled = TimeSliderController->IsSnapEnabled();
+	const bool bMiddleMouseDown = InMouseInput.PointerEvent.IsMouseButtonDown(EKeys::MiddleMouseButton);
+
+	const FFrameTime CurrentDragTime = GetAccumulatedTickTime(InMouseInput, bSnapEnabled);
+
+	TimeSliderController->CommitScrubPosition(CurrentDragTime, /*bIsScrubbing=*/true, !bMiddleMouseDown);
+}
+
+void FToolableTimelineScrubDragOperation::ResetDrag()
+{
+}
+
+} // UE::Sequencer::ToolableTimeline
