[ToolableTimeline] Give better name to instance settings customization namespace to avoid namespace abiguity
#Synced-CL 52513999

diff --git a/Engine/Source/Editor/Sequencer/Private/ToolableTimeline/ToolableTimelineInstanceSettingsCustomization.cpp b/Engine/Source/Editor/Sequencer/Private/ToolableTimeline/ToolableTimelineInstanceSettingsCustomization.cpp
--- a/Engine/Source/Editor/Sequencer/Private/ToolableTimeline/ToolableTimelineInstanceSettingsCustomization.cpp
+++ b/Engine/Source/Editor/Sequencer/Private/ToolableTimeline/ToolableTimelineInstanceSettingsCustomization.cpp
@@ -18,7 +18,7 @@
 namespace UE::Sequencer::ToolableTimeline
 {
 
-namespace Private
+namespace InstanceSettingsCustomization
 {
 	static const TCHAR* const EnableSimpleViewHiddenConfigCVarName = TEXT("Sequencer.SimpleView.HiddenConfig");
 	static TAutoConsoleVariable<bool> CVarEnableSimpleViewHiddenConfig(
@@ -31,7 +31,7 @@ namespace Private
 
 FToolableTimelineInstanceSettingsCustomization::~FToolableTimelineInstanceSettingsCustomization()
 {
-	Private::CVarEnableSimpleViewHiddenConfig->OnChangedDelegate().Remove(ViewHiddenConfigCVarChangedHandle);
+	InstanceSettingsCustomization::CVarEnableSimpleViewHiddenConfig->OnChangedDelegate().Remove(ViewHiddenConfigCVarChangedHandle);
 }
 
 TSharedRef<IPropertyTypeCustomization> FToolableTimelineInstanceSettingsCustomization::MakeInstance()
@@ -51,7 +51,7 @@ void FToolableTimelineInstanceSettingsCustomization::CustomizeChildren(TSharedRe
 
 	if (!ViewHiddenConfigCVarChangedHandle.IsValid())
 	{
-		ViewHiddenConfigCVarChangedHandle = Private::CVarEnableSimpleViewHiddenConfig->OnChangedDelegate().AddSP(this
+		ViewHiddenConfigCVarChangedHandle = InstanceSettingsCustomization::CVarEnableSimpleViewHiddenConfig->OnChangedDelegate().AddSP(this
 			, &FToolableTimelineInstanceSettingsCustomization::HandleHiddenConfigToggled);
 	}
 
@@ -122,7 +122,7 @@ void FToolableTimelineInstanceSettingsCustomization::CustomizeChildren(TSharedRe
 	{
 		const TSharedRef<IPropertyHandle> ChildHandle = InPropertyHandle->GetChildHandle(i).ToSharedRef();
 
-		if (!Private::CVarEnableSimpleViewHiddenConfig.GetValueOnAnyThread()
+		if (!InstanceSettingsCustomization::CVarEnableSimpleViewHiddenConfig.GetValueOnAnyThread()
 			&& ChildHandle->HasMetaData(TEXT("HiddenConfig")))
 		{
 			continue;
