[ToolableTimeline] Change namespace to be less ambiguous
#Synced-CL 52513999

diff --git a/Engine/Source/Editor/Sequencer/Private/ToolableTimeline/Menus/ToolableTimelineKeyContextMenu.cpp b/Engine/Source/Editor/Sequencer/Private/ToolableTimeline/Menus/ToolableTimelineKeyContextMenu.cpp
--- a/Engine/Source/Editor/Sequencer/Private/ToolableTimeline/Menus/ToolableTimelineKeyContextMenu.cpp
+++ b/Engine/Source/Editor/Sequencer/Private/ToolableTimeline/Menus/ToolableTimelineKeyContextMenu.cpp
@@ -17,7 +17,7 @@
 namespace UE::Sequencer::ToolableTimeline
 {
 
-namespace Private
+namespace KeyContextMenu
 {
 	void CreateKeyStructForSelection(TWeakPtr<ISequencer> InWeakSequencer
 		, TSharedPtr<FStructOnScope>& OutKeyStruct, TWeakObjectPtr<UMovieSceneSection>& OutKeyStructSection)
@@ -142,7 +142,7 @@ TSharedRef<SWidget> FToolableTimelineKeyContextMenu::GenerateWidget(const FMouse
 	{
 		TSharedPtr<FStructOnScope> KeyStruct;
 		TWeakObjectPtr<UMovieSceneSection> KeyStructSection;
-		Private::CreateKeyStructForSelection(WeakSequencer, KeyStruct, KeyStructSection);
+		KeyContextMenu::CreateKeyStructForSelection(WeakSequencer, KeyStruct, KeyStructSection);
 
 		if (KeyStruct.IsValid())
 		{
@@ -161,7 +161,7 @@ TSharedRef<SWidget> FToolableTimelineKeyContextMenu::GenerateWidget(const FMouse
 						.EditData_Lambda([WeakSequencer]
 							{
 								FKeyEditData EditData;
-								Private::CreateKeyStructForSelection(WeakSequencer, EditData.KeyStruct, EditData.OwningSection);
+								KeyContextMenu::CreateKeyStructForSelection(WeakSequencer, EditData.KeyStruct, EditData.OwningSection);
 								return EditData;
 							});
 					SubMenuBuilder.AddWidget(KeyEditWidget, FText::GetEmpty(), /*bInNoIndent=*/true);
@@ -179,7 +179,7 @@ TSharedRef<SWidget> FToolableTimelineKeyContextMenu::GenerateWidget(const FMouse
 	//    This is what gives you the built-in curve interpolation section.
 	ISequencerModule& SequencerModule = FModuleManager::LoadModuleChecked<ISequencerModule>(TEXT("Sequencer"));
 
-	TMap<FName, TArray<FExtendKeyMenuParams>> ChannelAndHandlesByType = Private::BuildSelectedKeyMenuParams(Sequencer);
+	TMap<FName, TArray<FExtendKeyMenuParams>> ChannelAndHandlesByType = KeyContextMenu::BuildSelectedKeyMenuParams(Sequencer);
 
 	for (TPair<FName, TArray<FExtendKeyMenuParams>>& Pair : ChannelAndHandlesByType)
 	{
