[ToolableTimeline] Fix ambiguous namespace for Horde error
#Synced-CL 52513999

diff --git a/Engine/Source/Editor/Sequencer/Private/Widgets/Sidebar/SequencerSelectionDrawer.cpp b/Engine/Source/Editor/Sequencer/Private/Widgets/Sidebar/SequencerSelectionDrawer.cpp
--- a/Engine/Source/Editor/Sequencer/Private/Widgets/Sidebar/SequencerSelectionDrawer.cpp
+++ b/Engine/Source/Editor/Sequencer/Private/Widgets/Sidebar/SequencerSelectionDrawer.cpp
@@ -95,8 +95,6 @@ namespace UE::Sequencer::Private
 	}
 }
 
-using namespace UE::Sequencer;
-
 const FName FSequencerSelectionDrawer::UniqueId = TEXT("SequencerSelectionDrawer");
 
 FSequencerSelectionDrawer::FSequencerSelectionDrawer(const TWeakPtr<FSequencer>& InWeakSequencer)
@@ -153,6 +151,8 @@ TSharedRef<SWidget> FSequencerSelectionDrawer::CreateContentWidget()
 
 void FSequencerSelectionDrawer::BindEvents()
 {
+	using namespace UE::Sequencer;
+
 	const TSharedPtr<ISequencer> Sequencer = WeakSequencer.Pin();
 	if (!Sequencer.IsValid())
 	{
@@ -175,7 +175,7 @@ void FSequencerSelectionDrawer::BindEvents()
 			}
 		});
 
-	if (const TSharedPtr<FSequencerSelection> SequencerSelection = Private::GetSelection(*Sequencer.Get()))
+	if (const TSharedPtr<FSequencerSelection> SequencerSelection = UE::Sequencer::Private::GetSelection(*Sequencer.Get()))
 	{
 		SequencerSelection->OnChanged.AddSP(this, &FSequencerSelectionDrawer::RequestSelectionUpdate);
 	}
@@ -183,6 +183,8 @@ void FSequencerSelectionDrawer::BindEvents()
 
 void FSequencerSelectionDrawer::UnbindEvents()
 {
+	using namespace UE::Sequencer;
+
 	const TSharedPtr<ISequencer> Sequencer = WeakSequencer.Pin();
 	if (!Sequencer.IsValid())
 	{
@@ -200,7 +202,7 @@ void FSequencerSelectionDrawer::UnbindEvents()
 		MovieSceneDataChangedDelegate.Reset();
 	}
 
-	if (const TSharedPtr<FSequencerSelection> SequencerSelection = Private::GetSelection(*Sequencer.Get()))
+	if (const TSharedPtr<FSequencerSelection> SequencerSelection = UE::Sequencer::Private::GetSelection(*Sequencer.Get()))
 	{
 		SequencerSelection->OnChanged.RemoveAll(this);
 	}
@@ -236,6 +238,8 @@ void FSequencerSelectionDrawer::RequestSelectionUpdate()
 
 void FSequencerSelectionDrawer::UpdateFromSelectionNextFrame()
 {
+	using namespace UE::Sequencer;
+
 	ResetContent();
 
 	const TSharedPtr<FSequencer> Sequencer = WeakSequencer.Pin();
@@ -244,7 +248,7 @@ void FSequencerSelectionDrawer::UpdateFromSelectionNextFrame()
 		return;
 	}
 
-	const TSharedPtr<FSequencerSelection> SequencerSelection = Private::GetSelection(*Sequencer.Get());
+	const TSharedPtr<FSequencerSelection> SequencerSelection = UE::Sequencer::Private::GetSelection(*Sequencer.Get());
 	if (!SequencerSelection.IsValid())
 	{
 		return;
@@ -317,7 +321,7 @@ void FSequencerSelectionDrawer::UpdateFromSelectionNextFrame()
 	}
 }
 
-void FSequencerSelectionDrawer::BuildKeySelectionDetails(const TSharedRef<FSequencerSelection>& InSelection, FMenuBuilder& MenuBuilder)
+void FSequencerSelectionDrawer::BuildKeySelectionDetails(const TSharedRef<UE::Sequencer::FSequencerSelection>& InSelection, FMenuBuilder& MenuBuilder)
 {
 	if (InSelection->KeySelection.Num() == 0)
 	{
@@ -337,8 +341,10 @@ void FSequencerSelectionDrawer::BuildKeySelectionDetails(const TSharedRef<FSeque
 	MenuBuilder.EndSection();
 }
 
-void FSequencerSelectionDrawer::BuildTrackAreaDetails(const TSharedRef<FSequencerSelection>& InSelection, FMenuBuilder& MenuBuilder)
+void FSequencerSelectionDrawer::BuildTrackAreaDetails(const TSharedRef<UE::Sequencer::FSequencerSelection>& InSelection, FMenuBuilder& MenuBuilder)
 {
+	using namespace UE::Sequencer;
+
 	TArray<TWeakObjectPtr<>> AllSectionObjects; 
 
 	for (const FViewModelPtr TrackAreaItem : InSelection->TrackArea)
@@ -359,8 +365,10 @@ void FSequencerSelectionDrawer::BuildTrackAreaDetails(const TSharedRef<FSequence
 	}
 }
 
-void FSequencerSelectionDrawer::BuildOutlinerDetails(const TSharedRef<FSequencerSelection>& InSelection, FMenuBuilder& MenuBuilder)
+void FSequencerSelectionDrawer::BuildOutlinerDetails(const TSharedRef<UE::Sequencer::FSequencerSelection>& InSelection, FMenuBuilder& MenuBuilder)
 {
+	using namespace UE::Sequencer;
+
 	if (InSelection->Outliner.Num() == 0)
 	{
 		return;
@@ -496,7 +504,7 @@ void FSequencerSelectionDrawer::BuildOutlinerDetails(const TSharedRef<FSequencer
 	}
 }
 
-void FSequencerSelectionDrawer::BuildMarkedFrameDetails(const TSharedRef<FSequencerSelection>& InSelection, FMenuBuilder& MenuBuilder)
+void FSequencerSelectionDrawer::BuildMarkedFrameDetails(const TSharedRef<UE::Sequencer::FSequencerSelection>& InSelection, FMenuBuilder& MenuBuilder)
 {
 	if (InSelection->MarkedFrames.Num() == 0)
 	{
@@ -514,8 +522,10 @@ void FSequencerSelectionDrawer::BuildMarkedFrameDetails(const TSharedRef<FSequen
 	MenuBuilder.EndSection();
 }
 
-void FSequencerSelectionDrawer::BuildExtensionDetails(const TSet<TViewModelPtr<FChannelGroupOutlinerModel>>& InChannelGroups, FMenuBuilder& MenuBuilder)
+void FSequencerSelectionDrawer::BuildExtensionDetails(const TSet<UE::Sequencer::TViewModelPtr<UE::Sequencer::FChannelGroupOutlinerModel>>& InChannelGroups, FMenuBuilder& MenuBuilder)
 {
+	using namespace UE::Sequencer;
+
 	const TSharedPtr<ISequencer> Sequencer = WeakSequencer.Pin();
 	if (!Sequencer.IsValid())
 	{
@@ -598,30 +608,34 @@ TSharedRef<SWidget> FSequencerSelectionDrawer::CreateNoSelectionHintText()
 
 FKeyEditData FSequencerSelectionDrawer::GetKeyEditData() const
 {
+	using namespace UE::Sequencer;
+
 	const TSharedPtr<ISequencer> Sequencer = WeakSequencer.Pin();
 	if (!Sequencer.IsValid())
 	{
 		return FKeyEditData();
 	}
 
-	const TSharedPtr<FSequencerSelection> SequencerSelection = Private::GetSelection(*Sequencer.Get());
+	const TSharedPtr<FSequencerSelection> SequencerSelection = UE::Sequencer::Private::GetSelection(*Sequencer.Get());
 	if (!SequencerSelection.IsValid())
 	{
 		return FKeyEditData();
 	}
 
-	return Private::GetKeyEditData(SequencerSelection->KeySelection);
+	return UE::Sequencer::Private::GetKeyEditData(SequencerSelection->KeySelection);
 }
 
-TSharedPtr<SWidget> FSequencerSelectionDrawer::CreateKeyFrameDetails(const TSharedRef<FSequencerSelection>& InSequencerSelection)
+TSharedPtr<SWidget> FSequencerSelectionDrawer::CreateKeyFrameDetails(const TSharedRef<UE::Sequencer::FSequencerSelection>& InSequencerSelection)
 {
+	using namespace UE::Sequencer;
+
 	const TSharedPtr<ISequencer> Sequencer = WeakSequencer.Pin();
 	if (!Sequencer.IsValid())
 	{
 		return nullptr;
 	}
 
-	const FKeyEditData KeyEditData = Private::GetKeyEditData(InSequencerSelection->KeySelection);
+	const FKeyEditData KeyEditData = UE::Sequencer::Private::GetKeyEditData(InSequencerSelection->KeySelection);
 	if (KeyEditData.KeyStruct.IsValid())
 	{
 		return SNew(SKeyEditInterface, Sequencer.ToSharedRef())
