[SequencerSimpleView] Replace GetSequencerChecked() with GetSequencer() and validity checks to guarantee no crashes from stale pointers.

#Synced-CL 52513999

diff --git a/Engine/Source/Editor/Sequencer/Private/SSequencer.cpp b/Engine/Source/Editor/Sequencer/Private/SSequencer.cpp
--- a/Engine/Source/Editor/Sequencer/Private/SSequencer.cpp
+++ b/Engine/Source/Editor/Sequencer/Private/SSequencer.cpp
@@ -2708,6 +2708,12 @@ SSequencer::~SSequencer()
 		}
 	}
 
+	if (SimpleViewTimeline.IsValid())
+	{
+		SimpleViewTimeline->Shutdown();
+		SimpleViewTimeline.Reset();
+	}
+
 	USelection::SelectionChangedEvent.RemoveAll(this);
 
 	TSharedPtr<FSequencer> Sequencer = SequencerPtr.Pin();
diff --git a/Engine/Source/Editor/Sequencer/Private/SimpleView/Menus/SequencerToolbarExtension.cpp b/Engine/Source/Editor/Sequencer/Private/SimpleView/Menus/SequencerToolbarExtension.cpp
--- a/Engine/Source/Editor/Sequencer/Private/SimpleView/Menus/SequencerToolbarExtension.cpp
+++ b/Engine/Source/Editor/Sequencer/Private/SimpleView/Menus/SequencerToolbarExtension.cpp
@@ -79,12 +79,18 @@ void FSequencerToolbarExtension::CreateKeyTranslateSection(FToolMenuSection& InS
 		return;
 	}
 
+	const TSharedPtr<FSequencer> Sequencer = Timeline->GetSequencer();
+	if (!Sequencer.IsValid())
+	{
+		return;
+	}
+
 	const TWeakPtr<FSimpleViewTimeline> WeakTimeline = Timeline;
 
 	const FSimpleViewCommands& SimpleViewCommands = FSimpleViewCommands::Get();
 
 	const TSharedRef<SValueStepper> ValueWidget = SNew(SValueStepper)
-		.NumericTypeInterface(Timeline->GetSequencerChecked()->GetNumericTypeInterface())
+		.NumericTypeInterface(Sequencer->GetNumericTypeInterface())
 		.MinFractionalDigits(0)
 		.IsEnabled_Lambda([WeakTimeline]()
 			{
diff --git a/Engine/Source/Editor/Sequencer/Private/SimpleView/SimpleViewTimeline.cpp b/Engine/Source/Editor/Sequencer/Private/SimpleView/SimpleViewTimeline.cpp
--- a/Engine/Source/Editor/Sequencer/Private/SimpleView/SimpleViewTimeline.cpp
+++ b/Engine/Source/Editor/Sequencer/Private/SimpleView/SimpleViewTimeline.cpp
@@ -176,7 +176,11 @@ void FSimpleViewTimeline::EnableSimpleView(const bool bInEnable)
 		return;
 	}
 
-	USequencerSettings& SequencerSettings = GetSequencerSettingsChecked();
+	USequencerSettings* const SequencerSettings = Sequencer->GetSequencerSettings();
+	if (!SequencerSettings)
+	{
+		return;
+	}
 
 	if (const TSharedPtr<SDockTab> ExistingTab = FSimpleViewTabAutosizeHelper::FindSequencerTab(Sequencer.ToSharedRef()))
 	{
@@ -227,7 +231,7 @@ void FSimpleViewTimeline::EnableSimpleView(const bool bInEnable)
 		}
 	}
 
-	SequencerSettings.SetIsSimpleView(bInEnable);
+	SequencerSettings->SetIsSimpleView(bInEnable);
 }
 
 void FSimpleViewTimeline::ToggleSimpleView()
@@ -303,7 +307,8 @@ bool FSimpleViewTimeline::HasSequencerKeySelection() const
 
 void FSimpleViewTimeline::TranslateKeyLeft()
 {
-	if (!HasSequencerKeySelection())
+	const TSharedPtr<FSequencer> Sequencer = GetSequencer();
+	if (!Sequencer.IsValid())
 	{
 		return;
 	}
@@ -311,13 +316,14 @@ void FSimpleViewTimeline::TranslateKeyLeft()
 	const FFrameNumber DeltaTime = FFrameTime::FromDecimal(KeyTranslateDelta).GetFrame();
 	if (DeltaTime != 0)
 	{
-		GetSequencerChecked()->TransformSelection(-DeltaTime, 1.f);
+		Sequencer->TransformSelection(-DeltaTime, 1.f);
 	}
 }
 
 void FSimpleViewTimeline::TranslateKeyRight()
 {
-	if (!HasSequencerKeySelection())
+	const TSharedPtr<FSequencer> Sequencer = GetSequencer();
+	if (!Sequencer.IsValid())
 	{
 		return;
 	}
@@ -325,33 +331,35 @@ void FSimpleViewTimeline::TranslateKeyRight()
 	const FFrameNumber DeltaTime = FFrameTime::FromDecimal(KeyTranslateDelta).GetFrame();
 	if (DeltaTime != 0)
 	{
-		GetSequencerChecked()->TransformSelection(DeltaTime, 1.f);
+		Sequencer->TransformSelection(DeltaTime, 1.f);
 	}
 }
 
 void FSimpleViewTimeline::ScaleKeyDivide()
 {
-	if (!HasSequencerKeySelection())
+	const TSharedPtr<FSequencer> Sequencer = GetSequencer();
+	if (!Sequencer.IsValid() || Sequencer->GetSelection().KeySelection.Num() == 0)
 	{
 		return;
 	}
 
 	if (KeyScaleFactor != 0.f && KeyScaleFactor != 1.f)
 	{
-		GetSequencerChecked()->TransformSelection(0, 1.f / KeyScaleFactor);
+		Sequencer->TransformSelection(0, 1.f / KeyScaleFactor);
 	}
 }
 
 void FSimpleViewTimeline::ScaleKeyMultiply()
 {
-	if (!HasSequencerKeySelection())
+	const TSharedPtr<FSequencer> Sequencer = GetSequencer();
+	if (!Sequencer.IsValid() || Sequencer->GetSelection().KeySelection.Num() == 0)
 	{
 		return;
 	}
 
 	if (KeyScaleFactor != 0.f && KeyScaleFactor != 1.f)
 	{
-		GetSequencerChecked()->TransformSelection(0, KeyScaleFactor);
+		Sequencer->TransformSelection(0, KeyScaleFactor);
 	}
 }
 
diff --git a/Engine/Source/Editor/Sequencer/Private/SimpleView/Widgets/SSequencerSimpleView.cpp b/Engine/Source/Editor/Sequencer/Private/SimpleView/Widgets/SSequencerSimpleView.cpp
--- a/Engine/Source/Editor/Sequencer/Private/SimpleView/Widgets/SSequencerSimpleView.cpp
+++ b/Engine/Source/Editor/Sequencer/Private/SimpleView/Widgets/SSequencerSimpleView.cpp
@@ -26,11 +26,6 @@ void SSequencerSimpleView::Construct(const FArguments& InArgs, const TSharedRef<
 {
 	Timeline = InTimeline;
 
-	const TSharedPtr<ISequencer> Sequencer = Timeline->GetSequencer();
-	check(Sequencer.IsValid());
-	const TSharedPtr<FSequencerEditorViewModel> SequencerViewModel = Sequencer->GetViewModel()->CastThisShared<FSequencerEditorViewModel>();
-	check(SequencerViewModel.IsValid());
-
 	const TSharedRef<ToolableTimeline::FToolableTimeSliderController> TimeRangeController = Timeline->GetTimeSliderController();
 
 	ISequencerWidgetsModule& SequencerWidgets = FModuleManager::Get().LoadModuleChecked<ISequencerWidgetsModule>(TEXT("SequencerWidgets"));
@@ -54,8 +49,15 @@ void SSequencerSimpleView::Reconstruct()
 		return;
 	}
 
-	const TSharedRef<ISequencer> Sequencer = Timeline->GetSequencerChecked();
-	const bool bTransportControlsVisible = IsPopoutTransportControlsVisible(Sequencer, SimpleViewSettings->bShowTransportControls);
+	const TSharedPtr<ISequencer> Sequencer = Timeline->GetSequencer();
+	if (!Sequencer.IsValid())
+	{
+		return;
+	}
+
+	const TSharedRef<ISequencer> SequencerRef = Sequencer.ToSharedRef();
+
+	const bool bTransportControlsVisible = IsPopoutTransportControlsVisible(SequencerRef, SimpleViewSettings->bShowTransportControls);
 
 	// If the time range controls are hidden, put them on instead of showing a bunch of empty space
 	const bool bTransportControlsNextToTimeline = bTransportControlsVisible
@@ -78,7 +80,7 @@ void SSequencerSimpleView::Reconstruct()
 				.VAlign(VAlign_Center)
 				.Padding(0.f, 0.f, 2.f, 0.f)
 				[
-					ConstructTransportControls(Sequencer)
+					ConstructTransportControls(SequencerRef)
 				]
 				+ SHorizontalBox::Slot()
 				.HAlign(HAlign_Fill)
@@ -113,7 +115,7 @@ void SSequencerSimpleView::Reconstruct()
 				.AutoWidth()
 				.VAlign(VAlign_Fill)
 				[
-					ConstructTransportControls(Sequencer)
+					ConstructTransportControls(SequencerRef)
 				]
 				+ SHorizontalBox::Slot()
 				.HAlign(HAlign_Fill)
@@ -146,8 +148,10 @@ void SSequencerSimpleView::Reconstruct()
 			{
 				if (const TSharedPtr<SSequencerSimpleView> This = ThisWeak.Pin())
 				{
-					const TSharedRef<ISequencer> Sequencer = This->Timeline->GetSequencerChecked();
-					This->PinnedTrackAreaWidget->SetOutliner(Sequencer->GetOutlinerViewWidget());
+					if (const TSharedPtr<ISequencer> Sequencer = This->Timeline->GetSequencer())
+					{
+						This->PinnedTrackAreaWidget->SetOutliner(Sequencer->GetOutlinerViewWidget());
+					}
 				}
 			});
 	}
@@ -211,10 +215,14 @@ TSharedRef<SBox> SSequencerSimpleView::ConstructTransportControls(const TSharedR
 		.Padding(0.f, 0.f, 2.f, 0.f)
 		.Visibility_Lambda([this]()
 			{
-				const TSharedRef<ISequencer> Sequencer = Timeline->GetSequencerChecked();
+				const TSharedPtr<ISequencer> Sequencer = Timeline->GetSequencer();
+				if (!Sequencer.IsValid())
+				{
+					return EVisibility::Collapsed;
+				}
 				const USequencerSimpleViewSettings* const SimpleViewSettings = GetDefault<USequencerSimpleViewSettings>();
 				const bool bShowTransportControls = SimpleViewSettings ? SimpleViewSettings->bShowTransportControls : true;
-				return IsPopoutTransportControlsVisible(Sequencer, bShowTransportControls)
+				return IsPopoutTransportControlsVisible(Sequencer.ToSharedRef(), bShowTransportControls)
 					? EVisibility::Visible : EVisibility::Collapsed;
 			})
 		[
@@ -234,7 +242,11 @@ FReply SSequencerSimpleView::OnKeyDown(const FGeometry& InGeometry, const FKeyEv
 		return FReply::Handled();
 	}
 
-	const TSharedRef<ISequencer> Sequencer = Timeline->GetSequencerChecked();
+	const TSharedPtr<ISequencer> Sequencer = Timeline->GetSequencer();
+	if (!Sequencer.IsValid())
+	{
+		return FReply::Handled();
+	}
 
 	if (const TSharedPtr<FUICommandList> CommandBindings = Sequencer->GetCommandBindings())
 	{
@@ -273,8 +285,11 @@ FReply SSequencerSimpleView::OnMouseButtonUp(const FGeometry& InGeometry, const
 
 TSharedPtr<SSequencer> SSequencerSimpleView::GetSequencerWidget() const
 {
-	const TSharedRef<ISequencer> Sequencer = Timeline->GetSequencerChecked();
-	return StaticCastSharedRef<SSequencer>(Sequencer->GetSequencerWidget());
+	if (const TSharedPtr<ISequencer> Sequencer = Timeline->GetSequencer())
+	{
+		return StaticCastSharedRef<SSequencer>(Sequencer->GetSequencerWidget());
+	}
+	return nullptr;
 }
 
 } // namespace UE::Sequencer::SimpleView
diff --git a/Engine/Source/Editor/Sequencer/Private/ToolableTimeline/Caches/ToolableTimelineChannelCache.cpp b/Engine/Source/Editor/Sequencer/Private/ToolableTimeline/Caches/ToolableTimelineChannelCache.cpp
--- a/Engine/Source/Editor/Sequencer/Private/ToolableTimeline/Caches/ToolableTimelineChannelCache.cpp
+++ b/Engine/Source/Editor/Sequencer/Private/ToolableTimeline/Caches/ToolableTimelineChannelCache.cpp
@@ -12,22 +12,28 @@ namespace UE::Sequencer::ToolableTimeline
 
 FToolableTimelineChannelCache::FToolableTimelineChannelCache(const FToolableTimeline& InTimeline)
 {
-	const FSequencer& Sequencer = *InTimeline.GetSequencerChecked();
+	const TSharedPtr<FSequencer> Sequencer = InTimeline.GetSequencer();
+	if (!Sequencer.IsValid())
+	{
+		return;
+	}
+
+	const FSequencer& SequencerRef = *Sequencer;
 
 	switch (InTimeline.GetTimelineSettings().Settings.KeyDisplay)
 	{
 	default:
 	case ETimelineKeyDisplay::SelectedAndPinned:
-		WeakChannelModels.Append(GetChannelsFromSelectedOutlinerItems(Sequencer));
-		WeakChannelModels.Append(GetChannelsFromPinnedOutlinerItems(Sequencer));
+		WeakChannelModels.Append(GetChannelsFromSelectedOutlinerItems(SequencerRef));
+		WeakChannelModels.Append(GetChannelsFromPinnedOutlinerItems(SequencerRef));
 		break;
 
 	case ETimelineKeyDisplay::Selected:
-		WeakChannelModels.Append(GetChannelsFromSelectedOutlinerItems(Sequencer));
+		WeakChannelModels.Append(GetChannelsFromSelectedOutlinerItems(SequencerRef));
 		break;
 
 	case ETimelineKeyDisplay::All:
-		WeakChannelModels.Append(GetChannelsWithKeys(Sequencer));
+		WeakChannelModels.Append(GetChannelsWithKeys(SequencerRef));
 		break;
 	}
 
diff --git a/Engine/Source/Editor/Sequencer/Private/ToolableTimeline/DragOperations/ToolableTimelineMarkDragOperation.cpp b/Engine/Source/Editor/Sequencer/Private/ToolableTimeline/DragOperations/ToolableTimelineMarkDragOperation.cpp
--- a/Engine/Source/Editor/Sequencer/Private/ToolableTimeline/DragOperations/ToolableTimelineMarkDragOperation.cpp
+++ b/Engine/Source/Editor/Sequencer/Private/ToolableTimeline/DragOperations/ToolableTimelineMarkDragOperation.cpp
@@ -40,7 +40,11 @@ void FToolableTimelineMarkDragOperation::ResetDrag()
 
 void FToolableTimelineMarkDragOperation::CacheMarks(const FMouseInputData& InMouseInput, const int32 InHitMarkIndex)
 {
-	const TSharedRef<FSequencer> Sequencer = InMouseInput.Timeline->GetSequencerChecked();
+	const TSharedPtr<FSequencer> Sequencer = InMouseInput.Timeline->GetSequencer();
+	if (!Sequencer.IsValid())
+	{
+		return;
+	}
 
 	const TSharedPtr<FSequencerEditorViewModel> ViewModel = Sequencer->GetViewModel();
 	if (!ViewModel.IsValid())
diff --git a/Engine/Source/Editor/Sequencer/Private/ToolableTimeline/Drawing/FrameRangeBubbleDrawer.cpp b/Engine/Source/Editor/Sequencer/Private/ToolableTimeline/Drawing/FrameRangeBubbleDrawer.cpp
--- a/Engine/Source/Editor/Sequencer/Private/ToolableTimeline/Drawing/FrameRangeBubbleDrawer.cpp
+++ b/Engine/Source/Editor/Sequencer/Private/ToolableTimeline/Drawing/FrameRangeBubbleDrawer.cpp
@@ -23,11 +23,15 @@ void FFrameRangeBubbleDrawer::CacheDrawData(const TRange<FFrameNumber>& InRange
 	, FMouseDrawInputData& MouseDrawInput
 	, const FText* const InOptionalText)
 {
-	NumericTypeInterface = MouseDrawInput.Timeline->GetSequencerChecked()->GetNumericTypeInterface();
-	TickResolution = MouseDrawInput.Timeline->GetSequencerChecked()->GetFocusedTickResolution();
-	DisplayRate = MouseDrawInput.Timeline->GetSequencerChecked()->GetFocusedDisplayRate();
+	const TSharedPtr<FSequencer> Sequencer = MouseDrawInput.Timeline->GetSequencer();
+	if (!Sequencer.IsValid())
+	{
+		return;
+	}
 
-	FSequencer& Sequencer = *MouseDrawInput.Timeline->GetSequencerChecked();
+	NumericTypeInterface = Sequencer->GetNumericTypeInterface();
+	TickResolution = Sequencer->GetFocusedTickResolution();
+	DisplayRate = Sequencer->GetFocusedDisplayRate();
 
 	FrameDifference = GetSignedFrameDifference(InRange);
 
@@ -67,17 +71,17 @@ void FFrameRangeBubbleDrawer::CacheDrawData(const TRange<FFrameNumber>& InRange
 				const FFrameTime TickTimeDifference = FFrameRate::TransformTime(FFrameTime(DisplayFrameDifference), DisplayRate, TickResolution);
 				const FFrameNumber TickFrameDifference = TickTimeDifference.RoundToFrame();
 
-				FrameString = FString::Printf(TEXT("%c %s"), Sign, *Utils::TickFrameToString(Sequencer
+				FrameString = FString::Printf(TEXT("%c %s"), Sign, *Utils::TickFrameToString(*Sequencer
 					, TickFrameDifference, DrawArgs.bPreventLeadingZeros, !bNegative));
 			}
 		}
 		else
 		{
 			FrameString = bSingleFrameRange
-				? *Utils::TickFrameToString(Sequencer, StartFrameInclusive, DrawArgs.bPreventLeadingZeros, false)
+				? *Utils::TickFrameToString(*Sequencer, StartFrameInclusive, DrawArgs.bPreventLeadingZeros, false)
 				: FString::Printf(TEXT("%s - %s")
-					, *Utils::TickFrameToString(Sequencer, StartFrameInclusive, DrawArgs.bPreventLeadingZeros, false)
-					, *Utils::TickFrameToString(Sequencer, EndLabelFrame, DrawArgs.bPreventLeadingZeros, false));
+					, *Utils::TickFrameToString(*Sequencer, StartFrameInclusive, DrawArgs.bPreventLeadingZeros, false)
+					, *Utils::TickFrameToString(*Sequencer, EndLabelFrame, DrawArgs.bPreventLeadingZeros, false));
 		}
 	}
 
diff --git a/Engine/Source/Editor/Sequencer/Private/ToolableTimeline/Drawing/ToolableTimelineDrawer.cpp b/Engine/Source/Editor/Sequencer/Private/ToolableTimeline/Drawing/ToolableTimelineDrawer.cpp
--- a/Engine/Source/Editor/Sequencer/Private/ToolableTimeline/Drawing/ToolableTimelineDrawer.cpp
+++ b/Engine/Source/Editor/Sequencer/Private/ToolableTimeline/Drawing/ToolableTimelineDrawer.cpp
@@ -264,8 +264,13 @@ int32 DrawBackground(FMouseDrawInputData& MouseDrawInput)
 
 int32 DrawTicks(FMouseDrawInputData& MouseDrawInput, FSequencerTimeSliderController::FDrawTickArgs& InArgs)
 {
+	const TSharedPtr<FSequencer> Sequencer = MouseDrawInput.Timeline->GetSequencer();
+	if (!Sequencer.IsValid())
+	{
+		return MouseDrawInput.LayerId;
+	}
+
 	const TSharedRef<FToolableTimeSliderController> TimeSliderController = MouseDrawInput.Timeline->GetTimeSliderController();
-	const TSharedRef<FSequencer> Sequencer = MouseDrawInput.Timeline->GetSequencerChecked();
 	const UToolableTimelineSettings& TimelineSettings = MouseDrawInput.Timeline->GetTimelineSettings();
 
 	const FFrameRate TickResolution = TimeSliderController->GetTickResolution();
@@ -689,7 +694,18 @@ int32 DrawKeys(FMouseDrawInputData& MouseDrawInput
 int32 DrawScrubHandle(FMouseDrawInputData& MouseDrawInput)
 {
 	const TSharedPtr<FSequencer> Sequencer = MouseDrawInput.Timeline->GetSequencer();
-	const TSharedPtr<FSequencerEditorViewModel> SequencerViewModel = Sequencer.IsValid() ? Sequencer->GetViewModel() : nullptr;
+	if (!Sequencer.IsValid())
+	{
+		return MouseDrawInput.LayerId;
+	}
+
+	const USequencerSettings* const SequencerSettings = Sequencer->GetSequencerSettings();
+	if (!SequencerSettings)
+	{
+		return MouseDrawInput.LayerId;
+	}
+
+	const TSharedPtr<FSequencerEditorViewModel> SequencerViewModel = Sequencer->GetViewModel();
 	if (!SequencerViewModel.IsValid())
 	{
 		return MouseDrawInput.LayerId;
@@ -702,7 +718,6 @@ int32 DrawScrubHandle(FMouseDrawInputData& MouseDrawInput)
 	}
 
 	const TSharedRef<FToolableTimeSliderController> TimeSliderController = MouseDrawInput.Timeline->GetTimeSliderController();
-	const USequencerSettings& SequencerSettings = MouseDrawInput.Timeline->GetSequencerSettingsChecked();
 	const UToolableTimelineSettings& TimelineSettings = MouseDrawInput.Timeline->GetTimelineSettings();
 	const FFrameRate TickResolution = TimeSliderController->GetTickResolution();
 	const FTimeSliderArgs& TimeSliderArgs = TimeSliderController->GetTimeSliderArgs();
@@ -719,7 +734,7 @@ int32 DrawScrubHandle(FMouseDrawInputData& MouseDrawInput)
 	const FTimeWarpTrackModel* const ActiveTimeWarpTrack = TimeWarpExtension
 		? TimeWarpExtension->GetActiveTimeWarpTrack() : nullptr;
 
-	const bool bIsWarped = SequencerSettings.GetTimeWarpDisplayMode() == ESequencerTimeWarpDisplay::WarpedTime
+	const bool bIsWarped = SequencerSettings->GetTimeWarpDisplayMode() == ESequencerTimeWarpDisplay::WarpedTime
 		&& ActiveTimeWarpTrack != nullptr;
 
 	FLinearColor ScrubColor = TimelineSettings.Settings.ScrubColor;
@@ -784,7 +799,7 @@ int32 DrawScrubHandle(FMouseDrawInputData& MouseDrawInput)
 	);
 
 	// Draw vertical sub frame line if we are not snapping to frame
-	if (!SequencerSettings.GetForceWholeFrames())
+	if (!SequencerSettings->GetForceWholeFrames())
 	{
 		const float ScrubPixelX = MouseDrawInput.RangeToScreen.InputToLocalX(ScrubPosition.AsSeconds());
 
@@ -869,6 +884,18 @@ int32 DrawScrubHandle(FMouseDrawInputData& MouseDrawInput)
 
 int32 DrawDragRangeIndicator(FMouseDrawInputData& MouseDrawInput)
 {
+	const TSharedPtr<FSequencer> Sequencer = MouseDrawInput.Timeline->GetSequencer();
+	if (!Sequencer.IsValid())
+	{
+		return MouseDrawInput.LayerId;
+	}
+
+	const USequencerSettings* const SequencerSettings = Sequencer->GetSequencerSettings();
+	if (!SequencerSettings)
+	{
+		return MouseDrawInput.LayerId;
+	}
+
 	const TSharedRef<FToolableTimeSliderController> TimeSliderController = MouseDrawInput.Timeline->GetTimeSliderController();
 	const TSharedPtr<FToolableTimelineBaseTool> ActiveTool = TimeSliderController->GetActiveTool();
 
@@ -892,8 +919,6 @@ int32 DrawDragRangeIndicator(FMouseDrawInputData& MouseDrawInput)
 		return MouseDrawInput.LayerId;
 	}
 
-	const TSharedRef<FSequencer> Sequencer = MouseDrawInput.Timeline->GetSequencerChecked();
-	const USequencerSettings& SequencerSettings = MouseDrawInput.Timeline->GetSequencerSettingsChecked();
 	const UToolableTimelineSettings& TimelineSettings = MouseDrawInput.Timeline->GetTimelineSettings();
 	const FMouseInputData MouseInput = MouseDrawInput.ToMouseInputData();
 
@@ -1001,7 +1026,7 @@ int32 DrawDragRangeIndicator(FMouseDrawInputData& MouseDrawInput)
 	DrawArgs.BubbleFontSize = 6.5f;
 	DrawArgs.BubbleBorderColor = FLinearColor::Black;
 	DrawArgs.bPreventLeadingZeros = true;
-	DrawArgs.bAlignToWholeFrame = SequencerSettings.GetForceWholeFrames();
+	DrawArgs.bAlignToWholeFrame = SequencerSettings->GetForceWholeFrames();
 
 	switch (TimelineSettings.Settings.LabelVerticalAlignment)
 	{
@@ -1040,12 +1065,22 @@ int32 DrawDragRangeIndicator(FMouseDrawInputData& MouseDrawInput)
 
 int32 DrawAreaViewScrubPosition(FMouseDrawInputData& MouseDrawInput, const bool bInDisplayScrubPosition)
 {
+	const TSharedPtr<FSequencer> Sequencer = MouseDrawInput.Timeline->GetSequencer();
+	if (!Sequencer.IsValid())
+	{
+		return MouseDrawInput.LayerId;
+	}
+
+	const USequencerSettings* const SequencerSettings = Sequencer->GetSequencerSettings();
+	if (!SequencerSettings)
+	{
+		return MouseDrawInput.LayerId;
+	}
+
 	static const FSlateBrush* const ScrubFillBrush = FAppStyle::GetBrush(TEXT("Sequencer.Timeline.ScrubFill"));
 
 	const TSharedRef<FToolableTimeSliderController> TimeSliderController = MouseDrawInput.Timeline->GetTimeSliderController();
-	const TSharedRef<ISequencer> Sequencer = MouseDrawInput.Timeline->GetSequencerChecked();
 	const FFrameRate TickResolution = TimeSliderController->GetTickResolution();
-	const USequencerSettings& SequencerSettings = MouseDrawInput.Timeline->GetSequencerSettingsChecked();
 	const UToolableTimelineSettings& TimelineSettings = MouseDrawInput.Timeline->GetTimelineSettings();
 	const FTimeSliderArgs& TimeSliderArgs = TimeSliderController->GetTimeSliderArgs();
 
@@ -1056,9 +1091,9 @@ int32 DrawAreaViewScrubPosition(FMouseDrawInputData& MouseDrawInput, const bool
 
 	// If we have no active timewarp track, just draw the single scrub position
 	if (!ActiveTimeWarpTrack
-		|| SequencerSettings.GetTimeWarpDisplayMode() != ESequencerTimeWarpDisplay::Both)
+		|| SequencerSettings->GetTimeWarpDisplayMode() != ESequencerTimeWarpDisplay::Both)
 	{
-		const bool bIsWarped = (SequencerSettings.GetTimeWarpDisplayMode() == ESequencerTimeWarpDisplay::WarpedTime
+		const bool bIsWarped = (SequencerSettings->GetTimeWarpDisplayMode() == ESequencerTimeWarpDisplay::WarpedTime
 			&& ActiveTimeWarpTrack != nullptr);
 		const FLinearColor ScrubColor = bIsWarped ? TimeWarpColor : TimelineSettings.Settings.ScrubColor;
 
diff --git a/Engine/Source/Editor/Sequencer/Private/ToolableTimeline/Menus/ToolableTimelineKeyContextMenu.cpp b/Engine/Source/Editor/Sequencer/Private/ToolableTimeline/Menus/ToolableTimelineKeyContextMenu.cpp
--- a/Engine/Source/Editor/Sequencer/Private/ToolableTimeline/Menus/ToolableTimelineKeyContextMenu.cpp
+++ b/Engine/Source/Editor/Sequencer/Private/ToolableTimeline/Menus/ToolableTimelineKeyContextMenu.cpp
@@ -11,6 +11,7 @@
 #include "SequencerToolMenuContext.h"
 #include "SKeyEditInterface.h"
 #include "ToolableTimeline/MouseInputData.h"
+#include "Widgets/SNullWidget.h"
 
 #define LOCTEXT_NAMESPACE "ToolableTimelineKeyContextMenu"
 
@@ -132,7 +133,12 @@ namespace KeyContextMenu
 
 TSharedRef<SWidget> FToolableTimelineKeyContextMenu::GenerateWidget(const FMouseInputData& InMouseInput)
 {
-	const TSharedRef<FSequencer> Sequencer = InMouseInput.Timeline->GetSequencerChecked();
+	const TSharedPtr<FSequencer> Sequencer = InMouseInput.Timeline->GetSequencer();
+	if (!Sequencer.IsValid())
+	{
+		return SNullWidget::NullWidget;
+	}
+
 	const TWeakPtr<FSequencer> WeakSequencer = Sequencer;
 
 	const TSharedPtr<FExtender> MenuExtender = MakeShared<FExtender>();
@@ -179,7 +185,7 @@ TSharedRef<SWidget> FToolableTimelineKeyContextMenu::GenerateWidget(const FMouse
 	//    This is what gives you the built-in curve interpolation section.
 	ISequencerModule& SequencerModule = FModuleManager::LoadModuleChecked<ISequencerModule>(TEXT("Sequencer"));
 
-	TMap<FName, TArray<FExtendKeyMenuParams>> ChannelAndHandlesByType = KeyContextMenu::BuildSelectedKeyMenuParams(Sequencer);
+	TMap<FName, TArray<FExtendKeyMenuParams>> ChannelAndHandlesByType = KeyContextMenu::BuildSelectedKeyMenuParams(Sequencer.ToSharedRef());
 
 	for (TPair<FName, TArray<FExtendKeyMenuParams>>& Pair : ChannelAndHandlesByType)
 	{
diff --git a/Engine/Source/Editor/Sequencer/Private/ToolableTimeline/ToolableTimeSliderController.cpp b/Engine/Source/Editor/Sequencer/Private/ToolableTimeline/ToolableTimeSliderController.cpp
--- a/Engine/Source/Editor/Sequencer/Private/ToolableTimeline/ToolableTimeSliderController.cpp
+++ b/Engine/Source/Editor/Sequencer/Private/ToolableTimeline/ToolableTimeSliderController.cpp
@@ -99,6 +99,18 @@ int32 FToolableTimeSliderController::OnPaintTimeSlider(const bool bInMirrorLabel
 	, const FWidgetStyle& InWidgetStyle
 	, const bool bInParentEnabled) const
 {
+	const TSharedPtr<ISequencer> Sequencer = Timeline.GetSequencer();
+	if (!Sequencer.IsValid())
+	{
+		return LayerId;
+	}
+
+	const USequencerSettings* const SequencerSettings = Sequencer->GetSequencerSettings();
+	if (!SequencerSettings)
+	{
+		return LayerId;
+	}
+
 	const TRange<double> ViewRange = GetViewRange();
 	const float LocalSequenceLength = ViewRange.GetUpperBoundValue() - ViewRange.GetLowerBoundValue();
 	if (LocalSequenceLength <= 0.f)
@@ -106,10 +118,7 @@ int32 FToolableTimeSliderController::OnPaintTimeSlider(const bool bInMirrorLabel
 		return LayerId;
 	}
 
-	const TSharedRef<ISequencer> Sequencer = Timeline.GetSequencerChecked();
-	const USequencerSettings& SequencerSettings = Timeline.GetSequencerSettingsChecked();
 	const UToolableTimelineSettings& TimelineSettings = Timeline.GetTimelineSettings();
-
 	const ESlateDrawEffect DrawEffects = bInParentEnabled ? ESlateDrawEffect::None : ESlateDrawEffect::DisabledEffect;
 	FMouseDrawInputData MouseDrawInput(Timeline.AsShared(), FPointerEvent()
 		, InGeometry, InCullingRect, OutDrawElements, LayerId, DrawEffects, bInMirrorLabels);
@@ -216,7 +225,7 @@ int32 FToolableTimeSliderController::OnPaintTimeSlider(const bool bInMirrorLabel
 		PlaybackRangeArgs
 	);
 
-	if (SequencerSettings.GetShowMarkedFrames())
+	if (SequencerSettings->GetShowMarkedFrames())
 	{
 		MouseDrawInput.LayerId = DrawMarkedFrames(InGeometry,
 			MouseDrawInput.RangeToScreen,
@@ -252,10 +261,19 @@ int32 FToolableTimeSliderController::OnPaintViewArea(const FGeometry& InGeometry
 	, bool bInEnabled
 	, const FPaintViewAreaArgs& InArgs) const
 {
-	const TSharedRef<ISequencer> Sequencer = Timeline.GetSequencerChecked();
-	const USequencerSettings& SequencerSettings = Timeline.GetSequencerSettingsChecked();
-	const UToolableTimelineSettings& TimelineSettings = Timeline.GetTimelineSettings();
+	const TSharedPtr<ISequencer> Sequencer = Timeline.GetSequencer();
+	if (!Sequencer.IsValid())
+	{
+		return LayerId;
+	}
 
+	const USequencerSettings* const SequencerSettings = Sequencer->GetSequencerSettings();
+	if (!SequencerSettings)
+	{
+		return LayerId;
+	}
+
+	const UToolableTimelineSettings& TimelineSettings = Timeline.GetTimelineSettings();
 	const ESlateDrawEffect DrawEffects = bInEnabled ? ESlateDrawEffect::None : ESlateDrawEffect::DisabledEffect;
 	FMouseDrawInputData MouseDrawInput(Timeline.AsShared(), FPointerEvent()
 		, InGeometry, InCullingRect, OutDrawElements, LayerId, DrawEffects);
@@ -309,7 +327,7 @@ int32 FToolableTimeSliderController::OnPaintViewArea(const FGeometry& InGeometry
 		Drawing::DrawTicks(MouseDrawInput, DrawTickArgs);
 	}
 
-	if (SequencerSettings.GetShowMarkedFrames())
+	if (SequencerSettings->GetShowMarkedFrames())
 	{
 		MouseDrawInput.LayerId = DrawMarkedFrames(InGeometry,
 			MouseDrawInput.RangeToScreen,
@@ -596,7 +614,12 @@ FReply FToolableTimeSliderController::OnMouseMoveImpl(SWidget& OwnerWidget
 	, const FPointerEvent& InPointerEvent
 	, const bool bInFromTimeSlider)
 {
-	const USequencerSettings& SequencerSettings = Timeline.GetSequencerSettingsChecked();
+	const USequencerSettings* const SequencerSettings = Timeline.GetSequencerSettings();
+	if (!SequencerSettings)
+	{
+		return FReply::Unhandled();
+	}
+
 	const FMouseInputData MouseInput(Timeline.AsShared(), InGeometry, InPointerEvent);
 
 	FReply OutReply = FReply::Unhandled();
@@ -663,7 +686,7 @@ FReply FToolableTimeSliderController::OnMouseMoveImpl(SWidget& OwnerWidget
 
 	if (bRightMouseDown)
 	{
-		OutReply = OnMouseMove_RightMouseDown(InGeometry, InPointerEvent, SequencerSettings);
+		OutReply = OnMouseMove_RightMouseDown(InGeometry, InPointerEvent, *SequencerSettings);
 	}
 	else if ((bLeftMouseDown || bMiddleMouseDown) && MarkDragOperation.IsSet())
 	{
@@ -692,7 +715,7 @@ FReply FToolableTimeSliderController::OnMouseMoveImpl(SWidget& OwnerWidget
 			InPointerEvent,
 			bInFromTimeSlider,
 			MouseInput.RangeToScreen,
-			SequencerSettings
+			*SequencerSettings
 		);
 
 		// Scrub drag started this frame
@@ -957,7 +980,11 @@ bool FToolableTimeSliderController::IsScrubHandleHidden() const
 
 bool FToolableTimeSliderController::GetGridMetrics(const FGeometry& InGeometry, double& OutMajorGridStep, int32& OutMinorDivisions) const
 {
-	const TSharedRef<FSequencer> Sequencer = Timeline.GetSequencerChecked();
+	const TSharedPtr<FSequencer> Sequencer = Timeline.GetSequencer();
+	if (!Sequencer.IsValid())
+	{
+		return false;
+	}
 
 	const TSharedPtr<INumericTypeInterface<double>> NumericTypeInterface = Sequencer->GetNumericTypeInterface();
 	if (!NumericTypeInterface.IsValid())
@@ -989,14 +1016,18 @@ bool FToolableTimeSliderController::GetGridMetrics(const FGeometry& InGeometry,
 
 bool FToolableTimeSliderController::IsSnapEnabled() const
 {
-	const USequencerSettings& SequencerSettings = Timeline.GetSequencerSettingsChecked();
-	return SequencerSettings.GetIsSnapEnabled();
+	const USequencerSettings* const SequencerSettings = Timeline.GetSequencerSettings();
+	return SequencerSettings ? SequencerSettings->GetIsSnapEnabled() : false;
 }
 
 FToolableTimeSliderController::FScrubberMetrics FToolableTimeSliderController::GetScrubPixelMetrics(const FQualifiedFrameTime& InScrubTime
 	, const FScrubRangeToScreen& InRangeToScreen, const float InDilationPixels) const
 {
-	const TSharedRef<FSequencer> Sequencer = Timeline.GetSequencerChecked();
+	const TSharedPtr<FSequencer> Sequencer = Timeline.GetSequencer();
+	if (!Sequencer.IsValid())
+	{
+		return FScrubberMetrics();
+	}
 
 	static constexpr float MinScrubSize = 1.f;
 	static constexpr float HalfMinScrubSize = MinScrubSize * .5f;
@@ -1115,7 +1146,8 @@ FReply FToolableTimeSliderController::TryBeginMarkDrag(const FMouseInputData& In
 
 TSharedPtr<FSequencerEditorViewModel> FToolableTimeSliderController::GetSequencerViewModel() const
 {
-	return Timeline.GetSequencerChecked()->GetViewModel();
+	const TSharedPtr<FSequencer> Sequencer = Timeline.GetSequencer();
+	return Sequencer.IsValid() ? Sequencer->GetViewModel() : nullptr;
 }
 
 TSharedPtr<FSequencerSelection> FToolableTimeSliderController::GetSequencerSelection() const
@@ -1182,13 +1214,15 @@ FFrameTime FToolableTimeSliderController::ComputeMouseFrameTime(const FMouseInpu
 
 		if (bInCheckSnapping)
 		{
-			const TSharedRef<ISequencer> Sequencer = Timeline.GetSequencerChecked();
-
-			if (Sequencer->GetSequencerSettings()->GetForceWholeFrames())
+			if (const TSharedPtr<FSequencer> Sequencer = Timeline.GetSequencer())
 			{
-				const FFrameNumber SnappedFrameNumber = (MouseValue * GetDisplayRate()).FloorToFrame();
-				const FQualifiedFrameTime RoundedPlayFrame(SnappedFrameNumber, GetDisplayRate());
-				return RoundedPlayFrame.ConvertTo(GetTickResolution());
+				const USequencerSettings* const SequencerSettings = Sequencer->GetSequencerSettings();
+				if (SequencerSettings && SequencerSettings->GetForceWholeFrames())
+				{
+					const FFrameNumber SnappedFrameNumber = (MouseValue * GetDisplayRate()).FloorToFrame();
+					const FQualifiedFrameTime RoundedPlayFrame(SnappedFrameNumber, GetDisplayRate());
+					return RoundedPlayFrame.ConvertTo(GetTickResolution());
+				}
 			}
 		}
 
diff --git a/Engine/Source/Editor/Sequencer/Private/ToolableTimeline/ToolableTimeline.cpp b/Engine/Source/Editor/Sequencer/Private/ToolableTimeline/ToolableTimeline.cpp
--- a/Engine/Source/Editor/Sequencer/Private/ToolableTimeline/ToolableTimeline.cpp
+++ b/Engine/Source/Editor/Sequencer/Private/ToolableTimeline/ToolableTimeline.cpp
@@ -50,29 +50,30 @@ FToolableTimeline::~FToolableTimeline()
 void FToolableTimeline::Initialize()
 {
 	BindCommands();
-	
-	const TSharedRef<FSequencer> Sequencer = GetSequencerChecked();
 
-	Sequencer->OnMovieSceneDataChanged().AddSP(this, &FToolableTimeline::HandleMovieSceneDataChanged);
+	if (const TSharedPtr<FSequencer> Sequencer = GetSequencer())
+	{
+		Sequencer->OnMovieSceneDataChanged().AddSP(this, &FToolableTimeline::HandleMovieSceneDataChanged);
 
-	Sequencer->GetFilterBar()->OnFiltersChanged().AddLambda([this](const ESequencerFilterChange InChangeType, const TSharedRef<FSequencerTrackFilter>& InFilter)
-		{
-			RecacheChannelModels();
-		});
+		Sequencer->GetFilterBar()->OnFiltersChanged().AddLambda([this](const ESequencerFilterChange InChangeType, const TSharedRef<FSequencerTrackFilter>& InFilter)
+			{
+				RecacheChannelModels();
+			});
 
-	const TSharedRef<ISequencerTrackFilters> FilterInterface = Sequencer->GetFilterInterface();
-	FilterInterface->OnRequestUpdate().AddSP(this, &FToolableTimeline::RecacheChannelModels);
-	FilterInterface->OnTextFilterTextChanged().AddSP(this, &FToolableTimeline::RecacheChannelModels);
-	FilterInterface->OnMuteFiltersChanged().AddLambda([this](const bool bInMuted)
-		{
-			RecacheChannelModels();
-		});
+		const TSharedRef<ISequencerTrackFilters> FilterInterface = Sequencer->GetFilterInterface();
+		FilterInterface->OnRequestUpdate().AddSP(this, &FToolableTimeline::RecacheChannelModels);
+		FilterInterface->OnTextFilterTextChanged().AddSP(this, &FToolableTimeline::RecacheChannelModels);
+		FilterInterface->OnMuteFiltersChanged().AddLambda([this](const bool bInMuted)
+			{
+				RecacheChannelModels();
+			});
 
-	if (const TSharedPtr<FSequencerEditorViewModel> ViewModel = Sequencer->GetViewModel())
-	{
-		if (const TSharedPtr<FSequencerSelection> Selection = ViewModel->GetSelection())
+		if (const TSharedPtr<FSequencerEditorViewModel> ViewModel = Sequencer->GetViewModel())
 		{
-			Selection->Outliner.OnChanged.AddSP(this, &FToolableTimeline::HandleOutlinerSelectionChanged);
+			if (const TSharedPtr<FSequencerSelection> Selection = ViewModel->GetSelection())
+			{
+				Selection->Outliner.OnChanged.AddSP(this, &FToolableTimeline::HandleOutlinerSelectionChanged);
+			}
 		}
 	}
 
@@ -240,7 +241,11 @@ void FToolableTimeline::RemoveChannelFilter(const FDelegateHandle& InDelegateHan
 
 bool FToolableTimeline::PassesChannelFilters(const FChannelModel& InChannelModel) const
 {
-	const TSharedRef<FSequencer> Sequencer = GetSequencerChecked();
+	const TSharedPtr<ISequencer> Sequencer = GetSequencer();
+	if (!Sequencer.IsValid())
+	{
+		return false;
+	}
 
 	for (const FRegisteredChannelFilter& ChannelFilter : GetChannelFilters())
 	{
@@ -258,27 +263,12 @@ TSharedPtr<FSequencer> FToolableTimeline::GetSequencer() const
 	return WeakSequencer.Pin();
 }
 
-TSharedRef<FSequencer> FToolableTimeline::GetSequencerChecked() const
-{
-	const TSharedPtr<FSequencer> Sequencer = WeakSequencer.Pin();
-	check(Sequencer.IsValid());
-	return Sequencer.ToSharedRef();
-}
-
 USequencerSettings* FToolableTimeline::GetSequencerSettings() const
 {
 	const TSharedPtr<ISequencer> Sequencer = GetSequencer();
 	return Sequencer.IsValid() ? Sequencer->GetSequencerSettings() : nullptr;
 }
 
-USequencerSettings& FToolableTimeline::GetSequencerSettingsChecked() const
-{
-	const TSharedRef<ISequencer> Sequencer = GetSequencerChecked();
-	USequencerSettings* const SequencerSettings = Sequencer->GetSequencerSettings();
-	check(SequencerSettings);
-	return *SequencerSettings;
-}
-
 TSharedPtr<FTabManager> FToolableTimeline::GetTabManager() const
 {
 	const TSharedPtr<ISequencer> Sequencer = GetSequencer();
@@ -417,7 +407,11 @@ TArray<FFrameNumber> FToolableTimeline::GetAllChannelKeyTimes() const
 
 void FToolableTimeline::FrameAllKeys()
 {
-	const TSharedRef<ISequencer> Sequencer = GetSequencerChecked();
+	const TSharedPtr<ISequencer> Sequencer = GetSequencer();
+	if (!Sequencer.IsValid())
+	{
+		return;
+	}
 
 	const TArray<FFrameNumber> AllKeyTimes = GetAllChannelKeyTimes();
 	if (AllKeyTimes.IsEmpty())
@@ -557,7 +551,13 @@ void FToolableTimeline::CenterFrameOnScreen(const double InNewValue, const bool
 
 void FToolableTimeline::CutSelectedKeys()
 {
-	if (const TSharedPtr<FCurveEditor> CurveEditor = Utils::GetSequencerCurveEditor(*GetSequencerChecked()))
+	const TSharedPtr<ISequencer> Sequencer = GetSequencer();
+	if (!Sequencer.IsValid())
+	{
+		return;
+	}
+
+	if (const TSharedPtr<FCurveEditor> CurveEditor = Utils::GetSequencerCurveEditor(*Sequencer))
 	{
 		SyncSelectionToSequencerAndCurveEditor();
 
@@ -572,7 +572,13 @@ bool FToolableTimeline::CanCutSelectedKeys() const
 
 void FToolableTimeline::CopySelectedKeys()
 {
-	if (const TSharedPtr<FCurveEditor> CurveEditor = Utils::GetSequencerCurveEditor(*GetSequencerChecked()))
+	const TSharedPtr<ISequencer> Sequencer = GetSequencer();
+	if (!Sequencer.IsValid())
+	{
+		return;
+	}
+
+	if (const TSharedPtr<FCurveEditor> CurveEditor = Utils::GetSequencerCurveEditor(*Sequencer))
 	{
 		SyncSelectionToSequencerAndCurveEditor();
 
@@ -587,7 +593,13 @@ bool FToolableTimeline::CanCopySelectedKeys() const
 
 void FToolableTimeline::PasteSelectedKeys()
 {
-	if (const TSharedPtr<FCurveEditor> CurveEditor = Utils::GetSequencerCurveEditor(*GetSequencerChecked()))
+	const TSharedPtr<ISequencer> Sequencer = GetSequencer();
+	if (!Sequencer.IsValid())
+	{
+		return;
+	}
+
+	if (const TSharedPtr<FCurveEditor> CurveEditor = Utils::GetSequencerCurveEditor(*Sequencer))
 	{
 		CurveEditor->PasteKeys();
 	}
@@ -595,7 +607,13 @@ void FToolableTimeline::PasteSelectedKeys()
 
 bool FToolableTimeline::CanPasteSelectedKeys() const
 {
-	if (const TSharedPtr<FCurveEditor> CurveEditor = Utils::GetSequencerCurveEditor(*GetSequencerChecked()))
+	const TSharedPtr<ISequencer> Sequencer = GetSequencer();
+	if (!Sequencer.IsValid())
+	{
+		return false;
+	}
+
+	if (const TSharedPtr<FCurveEditor> CurveEditor = Utils::GetSequencerCurveEditor(*Sequencer))
 	{
 		FString TextToImport;
         FPlatformApplicationMisc::ClipboardPaste(TextToImport);
@@ -608,7 +626,13 @@ bool FToolableTimeline::CanPasteSelectedKeys() const
 
 void FToolableTimeline::DeleteSelectedKeys()
 {
-	if (const TSharedPtr<FCurveEditor> CurveEditor = Utils::GetSequencerCurveEditor(*GetSequencerChecked()))
+	const TSharedPtr<ISequencer> Sequencer = GetSequencer();
+	if (!Sequencer.IsValid())
+	{
+		return;
+	}
+
+	if (const TSharedPtr<FCurveEditor> CurveEditor = Utils::GetSequencerCurveEditor(*Sequencer))
 	{
 		SyncSelectionToSequencerAndCurveEditor();
 
@@ -623,7 +647,13 @@ bool FToolableTimeline::CanDeleteSelectedKeys() const
 
 void FToolableTimeline::SelectAllKeys()
 {
-	if (const TSharedPtr<FCurveEditor> CurveEditor = Utils::GetSequencerCurveEditor(*GetSequencerChecked()))
+	const TSharedPtr<ISequencer> Sequencer = GetSequencer();
+	if (!Sequencer.IsValid())
+	{
+		return;
+	}
+
+	if (const TSharedPtr<FCurveEditor> CurveEditor = Utils::GetSequencerCurveEditor(*Sequencer))
 	{
 		CurveEditor->SelectAllKeys();
 	}
@@ -631,7 +661,13 @@ void FToolableTimeline::SelectAllKeys()
 
 bool FToolableTimeline::ExecuteCurveEditorAction(const TSharedPtr<const FUICommandInfo>& InCommand)
 {
-	const TSharedPtr<FCurveEditor> CurveEditor = Utils::GetSequencerCurveEditor(*GetSequencerChecked());
+	const TSharedPtr<ISequencer> Sequencer = GetSequencer();
+	if (!Sequencer.IsValid())
+	{
+		return false;
+	}
+
+	const TSharedPtr<FCurveEditor> CurveEditor = Utils::GetSequencerCurveEditor(*Sequencer);
 	if (!CurveEditor.IsValid())
 	{
 		return false;
@@ -692,10 +728,16 @@ bool FToolableTimeline::CanSmartSnapKeys() const
 
 void FToolableTimeline::SetSelectionRangeFromToolRange()
 {
+	const TSharedPtr<FSequencer> Sequencer = GetSequencer();
+	if (!Sequencer.IsValid())
+	{
+		return;
+	}
+
 	const TSharedPtr<FToolableTimelineBaseTool> ActiveTool = GetActiveTool();
 	if (ActiveTool.IsValid() && ActiveTool->HasValidToolRange())
 	{
-		GetSequencerChecked()->SetSelectionRange(ActiveTool->GetToolRange());
+		Sequencer->SetSelectionRange(ActiveTool->GetToolRange());
 	}
 }
 
diff --git a/Engine/Source/Editor/Sequencer/Private/ToolableTimeline/ToolableTimeline.h b/Engine/Source/Editor/Sequencer/Private/ToolableTimeline/ToolableTimeline.h
--- a/Engine/Source/Editor/Sequencer/Private/ToolableTimeline/ToolableTimeline.h
+++ b/Engine/Source/Editor/Sequencer/Private/ToolableTimeline/ToolableTimeline.h
@@ -181,9 +181,7 @@ public:
 	bool PassesChannelFilters(const FChannelModel& InChannelModel) const;
 
 	TSharedPtr<FSequencer> GetSequencer() const;
-	TSharedRef<FSequencer> GetSequencerChecked() const;
 	USequencerSettings* GetSequencerSettings() const;
-	USequencerSettings& GetSequencerSettingsChecked() const;
 	TSharedPtr<FTabManager> GetTabManager() const;
 
 	FTimeRangeArgs& GetTimeRangeArgs() { return TimeRangeArgs; }
diff --git a/Engine/Source/Editor/Sequencer/Private/ToolableTimeline/Widgets/SToolableTimeline.cpp b/Engine/Source/Editor/Sequencer/Private/ToolableTimeline/Widgets/SToolableTimeline.cpp
--- a/Engine/Source/Editor/Sequencer/Private/ToolableTimeline/Widgets/SToolableTimeline.cpp
+++ b/Engine/Source/Editor/Sequencer/Private/ToolableTimeline/Widgets/SToolableTimeline.cpp
@@ -23,8 +23,6 @@ void SToolableTimeline::Construct(const FArguments& InArgs, const TSharedRef<FTo
 {
 	Timeline = InTimeline;
 
-	const TSharedPtr<ISequencer> Sequencer = Timeline->GetSequencer();
-	check(Sequencer.IsValid());
 	const TSharedRef<FToolableTimeSliderController> TimeSliderController = Timeline->GetTimeSliderController();
 
 	InputStack.AddHandler(&TimeSliderController.Get());
@@ -98,25 +96,26 @@ FReply SToolableTimeline::OnKeyDown(const FGeometry& InGeometry, const FKeyEvent
 		return FReply::Handled();
 	}
 
-	const TSharedRef<ISequencer> Sequencer = Timeline->GetSequencerChecked();
-
-	if (const TSharedPtr<FUICommandList> CommandBindings = Sequencer->GetCommandBindings())
+	if (const TSharedPtr<ISequencer> Sequencer = Timeline->GetSequencer())
 	{
-		if (CommandBindings->ProcessCommandBindings(InKeyEvent))
+		if (const TSharedPtr<FUICommandList> CommandBindings = Sequencer->GetCommandBindings())
 		{
-			return FReply::Handled();
+			if (CommandBindings->ProcessCommandBindings(InKeyEvent))
+			{
+				return FReply::Handled();
+			}
 		}
-	}
 
-	if (FCurveEditorExtension* const CurveEditorExtension = Sequencer->GetViewModel()->CastDynamic<FCurveEditorExtension>())
-	{
-		if (const TSharedPtr<FCurveEditor> CurveEditor = CurveEditorExtension->GetCurveEditor())
+		if (FCurveEditorExtension* const CurveEditorExtension = Sequencer->GetViewModel()->CastDynamic<FCurveEditorExtension>())
 		{
-			if (const TSharedPtr<FUICommandList> CurveEditorCommands = CurveEditor->GetCommands())
+			if (const TSharedPtr<FCurveEditor> CurveEditor = CurveEditorExtension->GetCurveEditor())
 			{
-				if (CurveEditorCommands->ProcessCommandBindings(InKeyEvent))
+				if (const TSharedPtr<FUICommandList> CurveEditorCommands = CurveEditor->GetCommands())
 				{
-					return FReply::Handled();
+					if (CurveEditorCommands->ProcessCommandBindings(InKeyEvent))
+					{
+						return FReply::Handled();
+					}
 				}
 			}
 		}
