Sequencer: Remove Sequencer.UseNewFilteringAlgorithm.

#Synced-CL 52544961

diff --git a/Engine/Source/Editor/Sequencer/Private/Filters/Filters/SequencerTrackFilter_Selected.cpp b/Engine/Source/Editor/Sequencer/Private/Filters/Filters/SequencerTrackFilter_Selected.cpp
--- a/Engine/Source/Editor/Sequencer/Private/Filters/Filters/SequencerTrackFilter_Selected.cpp
+++ b/Engine/Source/Editor/Sequencer/Private/Filters/Filters/SequencerTrackFilter_Selected.cpp
@@ -34,50 +34,6 @@ static TArray<TWeakObjectPtr<USequencerTrackFilterSelectedExtension>> GetSelecti
 	}
 	return Extensions;
 }
-
-/** @return Whether InItem is a child of InSelectedNode. */
-static bool IsPartOfHierarchy(const FViewModelPtr& InItem, const FViewModelPtr& InSelectedNode)
-{
-	if (!InSelectedNode.IsValid())
-	{
-		return false;
-	}
-	for (FParentFirstChildIterator It(InSelectedNode, /*bIncludeThis=*/true); It; ++It)
-	{
-		if (*It == InItem)
-		{
-			return true;
-		}
-	}
-	
-	return false;
-}
-/** @return Whether InItem is a child of the FViewModels enumerated by USequencerTrackFilterSelectedExtension. */
-static bool IsPartOfExtendedSelection(
-	const FViewModelPtr& InItem, 
-	const FSequencerTrackFilter_Selected& InFilter, 
-	TConstArrayView<TWeakObjectPtr<USequencerTrackFilterSelectedExtension>> InWeakExtensions
-	)
-{
-	using EBreakBehavior = USequencerTrackFilterSelectedExtension::EBreakBehavior;
-	const auto ProcessExtension = [&InItem, &InFilter](const TWeakObjectPtr<USequencerTrackFilterSelectedExtension>& WeakExtension)
-		{
-			if (!WeakExtension.IsValid())
-			{
-				return false;
-			}
-		
-			bool bIsPartOfSelection = false;
-			WeakExtension->EnumerateViewModelsConsideredAsSelected(InFilter, [&InItem, &bIsPartOfSelection](const FViewModelPtr& ConsiderSelected)
-			{
-				bIsPartOfSelection = IsPartOfHierarchy(InItem, ConsiderSelected);
-				return bIsPartOfSelection ? EBreakBehavior::Break : EBreakBehavior::Continue;
-			});
-			return bIsPartOfSelection;
-		};
-	
-	return Algo::AnyOf(InWeakExtensions, ProcessExtension);
-}
 }
 
 FSequencerTrackFilter_Selected::FSequencerTrackFilter_Selected(ISequencerTrackFilters& InFilterInterface, TSharedPtr<FFilterCategory> InCategory)
@@ -200,33 +156,6 @@ FString FSequencerTrackFilter_Selected::GetName() const
 
 FFilterResult FSequencerTrackFilter_Selected::Evaluate(FSequencerTrackFilterType InItem) const
 {
-	// TODO UE-371829: Remove this CVar
-	const IConsoleVariable* Variable = IConsoleManager::Get().FindConsoleVariable(TEXT("Sequencer.UseNewFilteringAlgorithm"));
-	const bool bUseOldAlgorithm = !ensure(Variable) || !Variable->GetBool();
-	if (bUseOldAlgorithm)
-	{
-		const TSharedPtr<FSequencerEditorViewModel> ViewModel = GetSequencer().GetViewModel();
-		if (!ViewModel.IsValid())
-		{
-			return EItemFilterState::Exclude;
-		}
- 
-		const TSharedPtr<FSequencerSelection> SequencerSelection = ViewModel->GetSelection();
-		if (!SequencerSelection.IsValid())
-		{
-			return EItemFilterState::Exclude;
-		}
-		const bool bIsPartOfSequencerSelection = Algo::AnyOf(SequencerSelection->Outliner, [InItem](const FViewModelPtr& Selected)
-		{
-			return SelectionDetail::IsPartOfHierarchy(InItem, Selected);
-		});
-		
-		const bool bPasses = bIsPartOfSequencerSelection 
-			// Allow external logic to make us pretend like additional items are selected without actually modifying FSequencerSelection::Outliner.
-			|| SelectionDetail::IsPartOfExtendedSelection(InItem, *this, CachedWeakExtensions);
-		return bPasses ? EItemFilterState::Include : EItemFilterState::Exclude;
-	}
-	
 	return EvaluateFilterCache(FilterCache, InItem);
 }
 
diff --git a/Engine/Source/Editor/Sequencer/Private/Filters/SequencerFilterBar.cpp b/Engine/Source/Editor/Sequencer/Private/Filters/SequencerFilterBar.cpp
--- a/Engine/Source/Editor/Sequencer/Private/Filters/SequencerFilterBar.cpp
+++ b/Engine/Source/Editor/Sequencer/Private/Filters/SequencerFilterBar.cpp
@@ -1424,20 +1424,9 @@ void FSequencerFilterBar::PreFilter() const
 	PreFilter(ModifiedFilter);
 }
 
-// TODO UE-371829: Remove this CVar: Also look for "Sequencer.UseNewFilteringAlgorithm" because it passed to FindConsoleVariable
-static TAutoConsoleVariable<bool> CVarUseNewFilteringAlgorithm(TEXT("Sequencer.UseNewFilteringAlgorithm"), true, TEXT(""));
-
 const FSequencerFilterData& FSequencerFilterBar::FilterNodes()
 {
-	if (CVarUseNewFilteringAlgorithm.GetValueOnAnyThread())
-	{
-		FilterNodesNew();
-	}
-	else
-	{
-		FilterNodesOld();
-	}
-	
+	EvaluateFilters();
 
 	// Always filter in spacer node
 	if (const FSequenceModel* const SequenceModel = Sequencer.GetNodeTree()->GetRootNode()->CastThis<FSequenceModel>())
@@ -1507,99 +1496,7 @@ FString FSequencerFilterBar::GenerateTextFilterStringFromEnabledFilters() const
 	return GeneratedFilterString;
 }
 
-void FSequencerFilterBar::FilterNodesOld()
-{
-	// Allow all filters to do any pre-processing before filtering
-	PreFilter();
-	
-	// Reset all filter data
-	FilterData.Reset();
-	
-	// Loop through all nodes and filter recursively
-	const bool bHasActiveFilter = HasAnyFilterActive();
-	for (const TViewModelPtr<IOutlinerExtension>& RootNode : Sequencer.GetNodeTree()->GetRootNodes())
-	{
-		FilterNodesRecursiveOld(bHasActiveFilter, RootNode);
-	}
-}
-
-bool FSequencerFilterBar::FilterNodesRecursiveOld(const bool bInHasActiveFilter, const TViewModelPtr<IOutlinerExtension>& InStartNode)
-{
-	TRACE_CPUPROFILER_EVENT_SCOPE(FSequencerFilterBar::FilterNodesRecursive);
-	
-	/**
-	 * Main Filtering Logic
-	 *
-	 * - Pinning overrides all other filters
-	 * - Hidden/Isolated tracks will take precedence over common filters
-	 * - Can hide sub tracks of isolated tracks
-	 */
-	bool bAnyChildPassed = false;
-	
-	// Child nodes should always be processed, as they may force their parents to pass
-	for (const TViewModelPtr<IOutlinerExtension>& Node : InStartNode.AsModel()->GetChildrenOfType<IOutlinerExtension>())
-	{
-		if (FilterNodesRecursiveOld(bInHasActiveFilter, Node))
-		{
-			bAnyChildPassed = true;
-		}
-	}
-	
-	// Increment the total node count so we can remove the code to loop again just to count
-	FilterData.IncrementTotalNodeCount();
-	// Early out if no filter
-	if (!bInHasActiveFilter)
-	{
-		FilterData.FilterInNode(InStartNode);
-		return false;
-	}
-	
-	const USequencerSettings* const SequencerSettings = Sequencer.GetSequencerSettings();
-	if (!SequencerSettings)
-	{
-		return false;
-	}
-	
-	// Pinning overrides all other filters
-	if (EnumHasAnyFlags(Flags, EFilterFlags::PinnedItemsCanSkipFiltering) && !SequencerSettings->GetIncludePinnedInFilter())
-	{
-		const TSharedPtr<IPinnableExtension> Pinnable = InStartNode.AsModel()->FindAncestorOfType<IPinnableExtension>(true);
-		if (Pinnable.IsValid() && Pinnable->IsPinned())
-		{
-			FilterData.FilterInParentChildNodes(InStartNode, ENodeVisitFlags::IncludeAll);
-			return true;
-		}
-	}
-	
-	const bool bPassedTextFilter = !TextFilter->IsActive() || TextFilter->PassesFilter(InStartNode);
-	const bool bPassedHideIsolateFilter = !HideIsolateFilter->IsActive() || HideIsolateFilter->PassesFilter(InStartNode);
-	const bool bPassedAnyCommonFilters = PassesAnyCommonFilter(InStartNode);
-	const bool bPassedInternalFilters = !AnyInternalFilterActive() || PassesAllInternalFilters(InStartNode);
-	const bool bPassedAnyCustomTextFilters = PassesAllCustomTextFilters(InStartNode);
-	const bool bAllFiltersPassed = bPassedTextFilter
-		&& bPassedHideIsolateFilter
-		&& bPassedAnyCommonFilters
-		&& bPassedInternalFilters
-		&& bPassedAnyCustomTextFilters;
-	
-	if (bAllFiltersPassed || bAnyChildPassed)
-	{
-		// Don't expand if a mouse capture is effecting TrackValueChanged events to avoid keys
-		// being hidden and expanding tracks while dragging
-		if (SequencerSettings->GetAutoExpandNodesOnFilterPass()
-			&& !FSlateApplication::Get().HasAnyMouseCaptor())
-		{
-			SetTrackParentsExpanded(InStartNode.ImplicitCast(), true);
-		}
-		FilterData.FilterInNodeWithAncestors(InStartNode);
-		return true;
-	}
-	// After child nodes are processed, fail anything that didn't pass
-	FilterData.FilterOutNode(InStartNode);
-	return false;
-}
-
-void FSequencerFilterBar::FilterNodesNew()
+void FSequencerFilterBar::EvaluateFilters()
 {
 	// Reset all filter data
 	FilterData.Reset();
@@ -1612,7 +1509,7 @@ void FSequencerFilterBar::FilterNodesNew()
 		// Loop through all nodes and filter recursively
 		for (const TViewModelPtr<IOutlinerExtension>& RootNode : Sequencer.GetNodeTree()->GetRootNodes())
 		{
-			FilterNodesRecursiveNew(RootNode, Filters, {});
+			FilterNodesRecursive(RootNode, Filters, {});
 		}
 	}
 	else
@@ -1624,7 +1521,7 @@ void FSequencerFilterBar::FilterNodesNew()
 	}
 }
 
-bool FSequencerFilterBar::FilterNodesRecursiveNew(
+bool FSequencerFilterBar::FilterNodesRecursive(
 	const TViewModelPtr<IOutlinerExtension>& InStartNode,
 	const FSequencerFilterContext& InFilters,
 	const FRecursiveFilterOverrides& InOverrides
@@ -1676,7 +1573,7 @@ bool FSequencerFilterBar::FilterNodesRecursiveNew(
 		// Child nodes should always be processed, as they may force their parents to pass
 		for (const TViewModelPtr<IOutlinerExtension>& Node : InStartNode.AsModel()->GetChildrenOfType<IOutlinerExtension>())
 		{
-			if (FilterNodesRecursiveNew(Node, InFilters, EvaluationResult.NewOverrides))
+			if (FilterNodesRecursive(Node, InFilters, EvaluationResult.NewOverrides))
 			{
 				bAnyChildPassed = true;
 			}
diff --git a/Engine/Source/Editor/Sequencer/Private/Filters/SequencerFilterBar.h b/Engine/Source/Editor/Sequencer/Private/Filters/SequencerFilterBar.h
--- a/Engine/Source/Editor/Sequencer/Private/Filters/SequencerFilterBar.h
+++ b/Engine/Source/Editor/Sequencer/Private/Filters/SequencerFilterBar.h
@@ -327,17 +327,12 @@ private:
 	/** Binds all commands that the filter bar or its filters handle. */
 	void BindCommands();
 	
-	/** Old algorithm for filtering nodes */
-	void FilterNodesOld();
-	/** Do not call directly! Should only be called by FilterNodes(). */
-	bool FilterNodesRecursiveOld(const bool bInHasActiveFilter, const UE::Sequencer::TViewModelPtr<UE::Sequencer::IOutlinerExtension>& InStartNode);
-	
 	/** New algorithm for filtering nodes. */
-	void FilterNodesNew();
+	void EvaluateFilters();
 	/** @return Filter that are active and should be evaluated. */
 	UE::Sequencer::FSequencerFilterContext BuildActiveFilters();
 	/** Do not call directly! Should only be called by FilterNodes(). */
-	bool FilterNodesRecursiveNew(
+	bool FilterNodesRecursive(
 		const UE::Sequencer::TViewModelPtr<UE::Sequencer::IOutlinerExtension>& InStartNode,
 		const UE::Sequencer::FSequencerFilterContext& InFilters,
 		const UE::Sequencer::FRecursiveFilterOverrides& InOverrides
