Sequencer: Fix items selected in viewport becoming deselected when they do not pass the filter in Outliner.

The issue was the Curve Editor's selection was getting updated by the Sequencer fitler: due to a missing scope guard, the Curve Editor would immediately sync its empty selection to Sequencer.

#Synced-CL 52544961

diff --git a/Engine/Source/Editor/Sequencer/Private/CurveEditor/Integration/Selection/CurveModelSyncer.cpp b/Engine/Source/Editor/Sequencer/Private/CurveEditor/Integration/Selection/CurveModelSyncer.cpp
--- a/Engine/Source/Editor/Sequencer/Private/CurveEditor/Integration/Selection/CurveModelSyncer.cpp
+++ b/Engine/Source/Editor/Sequencer/Private/CurveEditor/Integration/Selection/CurveModelSyncer.cpp
@@ -140,26 +140,91 @@
 void FCurveModelSyncer::UpdateCurveEditor()
 {
 	const TSharedPtr<FSequencer> Sequencer = WeakSequencer.Pin();
-	if (!Sequencer || SelectionSyncSuspension.IsSuspended())
+	FCurveEditorTree* CurveEditorTree = CurveEditor->GetTree();
+	if (!Sequencer || !CurveEditorTree || SelectionSyncSuspension.IsSuspended())
 	{
 		return;
 	}
 	
+	OnStartUpdateCurveEditorDelegate.Broadcast();
+	SyncSequencerToCurveEditor(*Sequencer, *CurveEditorTree);
+	OnStopUpdateCurveEditorDelegate.Broadcast();
+}
+
+FCurveEditorTreeItemID FCurveModelSyncer::AddToCurveEditor(
+	TViewModelPtr<ICurveEditorTreeItemExtension> InViewModel,
+	TSharedPtr<FCurveEditor> InCurveEditor
+	)
+{
+	// If the view model doesn't want to be in the curve editor, bail out
+	// Note that this means we will create curve editor items for each parent in the hierarchy up
+	// until the first parent that doesn't implement ICurveEditorTreeItemExtension
+	// That is: we don't create "dummy" entries when there's a "gap" in the hierarchy
+	if (!InViewModel)
+	{
+		return FCurveEditorTreeItemID::Invalid();
+	}
+
+	// Check if we already have a valid curve editor ID
+	if (FCurveEditorTreeItemID* Existing = ViewModelToTreeItemIDMap.Find(InViewModel.AsModel()))
+	{
+		if (InCurveEditor->GetTree()->FindItem(*Existing) != nullptr)
+		{
+			return *Existing;
+		}
+	}
+
+	// Recursively create any needed parent curve editor items
+	TViewModelPtr<ICurveEditorTreeItemExtension> Parent = InViewModel.AsModel()->CastParent<ICurveEditorTreeItemExtension>();
+
+	FCurveEditorTreeItemID ParentID = Parent ?
+		AddToCurveEditor(Parent, InCurveEditor) :
+		FCurveEditorTreeItemID::Invalid();
+
+	// Create the new curve editor item
+	FCurveEditorTreeItem* NewItem = InCurveEditor->AddTreeItem(ParentID);
+	TSharedPtr<ICurveEditorTreeItem> CurveEditorTreeItem = InViewModel->GetCurveEditorTreeItem();
+	NewItem->SetWeakItem(CurveEditorTreeItem);
+	TOptional<FString> UniquePathName =  InViewModel->GetUniquePathName();
+	NewItem->SetUniquePathName(UniquePathName);
+	
+	// Register the new ID in our map and notify the view model
+	ViewModelToTreeItemIDMap.Add(InViewModel.AsModel(), NewItem->GetID());
+	TreeItemIDToViewModelMap.Add(NewItem->GetID(), InViewModel.AsModel());
+	
+	InViewModel->OnAddedToCurveEditor(NewItem->GetID(), InCurveEditor);
+
+	return NewItem->GetID();
+}
+
+void FCurveModelSyncer::OnOutlinerSelectionChanged()
+{
+	const TSharedPtr<FSequencer> Sequencer = WeakSequencer.Pin();
+	USequencerSettings* Settings = Sequencer ? Sequencer->GetSequencerSettings() : nullptr;
+	if (Settings && Settings->ShouldSyncCurveEditorSelection())
+	{
+		// We schedule a selection synchronization for the next update. This synchronization must happen
+		// after all filters have been applied, because the items we want to select in the curve editor
+		// might be currently filtered out, but will be visible when filters are re-evaluated. This is
+		// why curve editor integration runs in FSequencerNodeTree after filtering.
+		RequestSyncSelection();
+	}
+}
+
+void FCurveModelSyncer::SyncSequencerToCurveEditor(FSequencer& InSequencer, FCurveEditorTree& InTree)
+{
+	// Guard against multiple broadcasts here and defer them until the end of this function
+	const FScopedCurveEditorTreeEventGuard ScopedEventGuard = InTree.ScopedEventGuard();
+	
 	// TODO UE-364008: In the future, we can remove the paths with bUseFilteredOut.
-	const bool bUseFilteredOut = Sequencer->GetSequencerSettings()->GetLinkFiltersWithCurveEditor();
+	const bool bUseFilteredOut = InSequencer.GetSequencerSettings()->GetLinkFiltersWithCurveEditor();
 	const auto IsFilteredOut = [this, bUseFilteredOut](const TViewModelPtr<IOutlinerExtension>& InItem)
-	{
-		return UseLinkedFilteringUI()
-			? LinkedFilterViewModel->IsFilteredOut(InItem)
-			: bUseFilteredOut && InItem && InItem->IsFilteredOut();
-	};
-
-	FCurveEditorTree* CurveEditorTree = CurveEditor->GetTree();
-
-	// Guard against multiple broadcasts here and defer them until the end of this function
-	FScopedCurveEditorTreeEventGuard ScopedEventGuard = CurveEditorTree->ScopedEventGuard();
-
-	// Remove any stale tree items
+		{
+			return UseLinkedFilteringUI()
+				? LinkedFilterViewModel->IsFilteredOut(InItem)
+				: bUseFilteredOut && InItem && InItem->IsFilteredOut();
+		};
+	
 	for (auto It = ViewModelToTreeItemIDMap.CreateIterator(); It; ++It)
 	{
 		bool bIsRelevant = false;
@@ -170,7 +235,7 @@
 			TViewModelPtr<IOutlinerExtension> OutlinerModel = ViewModel.ImplicitCast();
 			const bool bIsVisible = !IsFilteredOut(OutlinerModel);
 			
-			const FCurveEditorTreeItem* TreeItem = CurveEditorTree->FindItem(ItemID);
+			const FCurveEditorTreeItem* TreeItem = InTree.FindItem(ItemID);
 
 			FCurveEditorTreeItemID ParentID;
 			if (auto Parent = ViewModel.AsModel()->FindAncestorOfType<ICurveEditorTreeItemExtension>())
@@ -195,7 +260,7 @@
 	// Do a second pass to remove any items that were removed recursively above
 	for (auto It = ViewModelToTreeItemIDMap.CreateIterator(); It; ++It)
 	{
-		if (CurveEditorTree->FindItem(It->Value) == nullptr)
+		if (InTree.FindItem(It->Value) == nullptr)
 		{
 			TreeItemIDToViewModelMap.Remove(It->Value);
 			It.RemoveCurrent();
@@ -205,7 +270,7 @@
 	// Iterate all non-filtered out outliners and check for curve editor tree extensions
 	const bool bIncludeRootNode = false;
 	bool bItemsAdded = false;
-	const TSharedPtr<FSequenceModel> OwnerModel = Sequencer->GetViewModel()->GetRootSequenceModel();
+	const TSharedPtr<FSequenceModel> OwnerModel = InSequencer.GetViewModel()->GetRootSequenceModel();
 	for (TParentFirstChildIterator<IOutlinerExtension> It(OwnerModel, bIncludeRootNode); It; ++It)
 	{
 		if (IsFilteredOut(*It))
@@ -232,66 +297,6 @@
 	}
 }
 
-FCurveEditorTreeItemID FCurveModelSyncer::AddToCurveEditor(
-	TViewModelPtr<ICurveEditorTreeItemExtension> InViewModel,
-	TSharedPtr<FCurveEditor> InCurveEditor
-	)
-{
-	// If the view model doesn't want to be in the curve editor, bail out
-	// Note that this means we will create curve editor items for each parent in the hierarchy up
-	// until the first parent that doesn't implement ICurveEditorTreeItemExtension
-	// That is: we don't create "dummy" entries when there's a "gap" in the hierarchy
-	if (!InViewModel)
-	{
-		return FCurveEditorTreeItemID::Invalid();
-	}
-
-	// Check if we already have a valid curve editor ID
-	if (FCurveEditorTreeItemID* Existing = ViewModelToTreeItemIDMap.Find(InViewModel.AsModel()))
-	{
-		if (InCurveEditor->GetTree()->FindItem(*Existing) != nullptr)
-		{
-			return *Existing;
-		}
-	}
-
-	// Recursively create any needed parent curve editor items
-	TViewModelPtr<ICurveEditorTreeItemExtension> Parent = InViewModel.AsModel()->CastParent<ICurveEditorTreeItemExtension>();
-
-	FCurveEditorTreeItemID ParentID = Parent ?
-		AddToCurveEditor(Parent, InCurveEditor) :
-		FCurveEditorTreeItemID::Invalid();
-
-	// Create the new curve editor item
-	FCurveEditorTreeItem* NewItem = InCurveEditor->AddTreeItem(ParentID);
-	TSharedPtr<ICurveEditorTreeItem> CurveEditorTreeItem = InViewModel->GetCurveEditorTreeItem();
-	NewItem->SetWeakItem(CurveEditorTreeItem);
-	TOptional<FString> UniquePathName =  InViewModel->GetUniquePathName();
-	NewItem->SetUniquePathName(UniquePathName);
-	
-	// Register the new ID in our map and notify the view model
-	ViewModelToTreeItemIDMap.Add(InViewModel.AsModel(), NewItem->GetID());
-	TreeItemIDToViewModelMap.Add(NewItem->GetID(), InViewModel.AsModel());
-	
-	InViewModel->OnAddedToCurveEditor(NewItem->GetID(), InCurveEditor);
-
-	return NewItem->GetID();
-}
-
-void FCurveModelSyncer::OnOutlinerSelectionChanged()
-{
-	const TSharedPtr<FSequencer> Sequencer = WeakSequencer.Pin();
-	USequencerSettings* Settings = Sequencer ? Sequencer->GetSequencerSettings() : nullptr;
-	if (Settings && Settings->ShouldSyncCurveEditorSelection())
-	{
-		// We schedule a selection synchronization for the next update. This synchronization must happen
-		// after all filters have been applied, because the items we want to select in the curve editor
-		// might be currently filtered out, but will be visible when filters are re-evaluated. This is
-		// why curve editor integration runs in FSequencerNodeTree after filtering.
-		RequestSyncSelection();
-	}
-}
-
 void FCurveModelSyncer::CleanUpPendingSyncSelection()
 {
 	bSyncSelectionRequested = false;
diff --git a/Engine/Source/Editor/Sequencer/Private/CurveEditor/Integration/Selection/CurveModelSyncer.h b/Engine/Source/Editor/Sequencer/Private/CurveEditor/Integration/Selection/CurveModelSyncer.h
--- a/Engine/Source/Editor/Sequencer/Private/CurveEditor/Integration/Selection/CurveModelSyncer.h
+++ b/Engine/Source/Editor/Sequencer/Private/CurveEditor/Integration/Selection/CurveModelSyncer.h
@@ -10,6 +10,7 @@
 #include "Templates/UnrealTemplate.h"
 
 class FSequencer;
+class FCurveEditorTree;
 
 namespace UE::Sequencer
 {
@@ -54,6 +55,10 @@ public:
 	
 	/** Invoked to request that the UI be updated with the current selection. */
 	FSimpleMulticastDelegate& OnRequestSyncCurveEditorUI() { return OnRequestSyncCurveEditorUIDelegate; }
+	/** Invoked before the curve editor is updated by UpdateCurveEditor. */
+	FSimpleMulticastDelegate& OnStartUpdateCurveEditor() { return OnStartUpdateCurveEditorDelegate; }
+	/** Invoked after the curve editor is updated by UpdateCurveEditor. */
+	FSimpleMulticastDelegate& OnStopUpdateCurveEditor() { return OnStopUpdateCurveEditorDelegate; }
 	
 private:
 	
@@ -72,6 +77,10 @@ private:
 	FTimerHandle SyncSelectionHandle;
 	/** Invoked to request that the UI be updated with the current selection. */
 	FSimpleMulticastDelegate OnRequestSyncCurveEditorUIDelegate;
+	/** Invoked before the curve editor is updated by UpdateCurveEditor. */
+	FSimpleMulticastDelegate OnStartUpdateCurveEditorDelegate;
+	/** Invoked after the curve editor is updated by UpdateCurveEditor. */
+	FSimpleMulticastDelegate OnStopUpdateCurveEditorDelegate;
 	
 	/** Suspends syncing selection to Curve Editor. */
 	FSuspendCounter SelectionSyncSuspension;
@@ -86,6 +95,9 @@ private:
 	/** Syncs Curve Editor's selection if selection syncing is enabled. */
 	void OnOutlinerSelectionChanged();
 	
+	/** Applies Sequencer's selection to the curve editor tree. */
+	void SyncSequencerToCurveEditor(FSequencer& InSequencer, FCurveEditorTree& InTree);
+	
 	/** Cancels any pending transactions. */
 	void CleanUpPendingSyncSelection();
 };
diff --git a/Engine/Source/Editor/Sequencer/Private/CurveEditor/Integration/Selection/SelectionSyncRootLogic.cpp b/Engine/Source/Editor/Sequencer/Private/CurveEditor/Integration/Selection/SelectionSyncRootLogic.cpp
--- a/Engine/Source/Editor/Sequencer/Private/CurveEditor/Integration/Selection/SelectionSyncRootLogic.cpp
+++ b/Engine/Source/Editor/Sequencer/Private/CurveEditor/Integration/Selection/SelectionSyncRootLogic.cpp
@@ -18,11 +18,20 @@ FSelectionSyncRootLogic::FSelectionSyncRootLogic(
 	, SequencerToCurveEditorSyncer(InSequencer, InCurveEditor, InLinkedFilteringViewModel)
 	, CurveEditorToSequencerSyncer(SequencerToCurveEditorSyncer, InCurveEditor, InSequencer)
 {
+	// If Sequencer initiates a sync from Sequencer to Curve Editor, prevent Curve Editor from syncing its select back to Sequencer.
 	SequencerToCurveEditorSyncer.OnRequestSyncCurveEditorUI().AddLambda([this]
 	{
 		const FScopedSuspendCurveEditorToSequencerSync Suspend = CurveEditorToSequencerSyncer.SuspendSync();
 		OnRequestRefreshCurveEditorUIDelegate.Broadcast();
 	});
+	SequencerToCurveEditorSyncer.OnStartUpdateCurveEditor().AddLambda([this]
+	{
+		SuspendDuringUpdateCurveEditor = CurveEditorToSequencerSyncer.SuspendSync();
+	});
+	SequencerToCurveEditorSyncer.OnStopUpdateCurveEditor().AddLambda([this]
+	{
+		SuspendDuringUpdateCurveEditor.Reset();
+	});
 	
 	CurveEditorToSequencerSyncer.OnStartSync().AddRaw(this, &FSelectionSyncRootLogic::OnStartCurveEditorToSequencerSync);
 	CurveEditorToSequencerSyncer.OnEndSync().AddRaw(this, &FSelectionSyncRootLogic::OnEndCurveEditorToSequencerSync);
@@ -69,7 +78,7 @@ void FSelectionSyncRootLogic::UpdateIsolationFilter()
 
 void FSelectionSyncRootLogic::OnStartCurveEditorToSequencerSync()
 {
-	SuspendSequenceToCurveEditor.Emplace(SequencerToCurveEditorSyncer.SuspendSelectionSync());
+	SuspendDuringRequestUpdateUI.Emplace(SequencerToCurveEditorSyncer.SuspendSelectionSync());
 			
 	// We don't want Curve Editor items to disappear as result of the user selecting something in the Curve Editor tree.
 	bSkipIsolation = true;
@@ -77,7 +86,7 @@ void FSelectionSyncRootLogic::OnStartCurveEditorToSequencerSync()
 
 void FSelectionSyncRootLogic::OnEndCurveEditorToSequencerSync()
 {
-	SuspendSequenceToCurveEditor.Reset();
+	SuspendDuringRequestUpdateUI.Reset();
 	bSkipIsolation = false;
 	
 	// Applying Curve Editor selection to the Outliner may cause items to become hidden in the Outliner, so refresh so the items are up to date. 
diff --git a/Engine/Source/Editor/Sequencer/Private/CurveEditor/Integration/Selection/SelectionSyncRootLogic.h b/Engine/Source/Editor/Sequencer/Private/CurveEditor/Integration/Selection/SelectionSyncRootLogic.h
--- a/Engine/Source/Editor/Sequencer/Private/CurveEditor/Integration/Selection/SelectionSyncRootLogic.h
+++ b/Engine/Source/Editor/Sequencer/Private/CurveEditor/Integration/Selection/SelectionSyncRootLogic.h
@@ -51,8 +51,10 @@ private:
 	
 	/** Handles that Sequencer has the same selection as that of Curve Editor. */
 	FSequencerSelectionSyncer CurveEditorToSequencerSyncer;
-	/** Set while CurveEditorToSequencerSyncer is syncing. */
-	TOptional<FScopedSuspendSequencerToCurveEditorSync> SuspendSequenceToCurveEditor;
+	/** Set while CurveEditorToSequencerSyncer is syncing due to OnRequestSyncCurveEditorUI. */
+	TOptional<FScopedSuspendSequencerToCurveEditorSync> SuspendDuringRequestUpdateUI;
+	/** Set while CurveEditorToSequencerSyncer is syncing due to a UpdateCurveEditor call. */
+	TOptional<FScopedSuspendSequencerToCurveEditorSync> SuspendDuringUpdateCurveEditor;
 	
 	/** Curve editor filter that shows only the nodes selected in the Outliner. Can be null. */
 	TSharedPtr<FSequencerSelectionCurveFilter> SelectionCurveEditorFilter;
