Sequencer: Removes linked filtering CVars and USequencerSettings::bLinkFiltersWithCurveEditor.

#Synced-CL 52617198

diff --git a/Engine/Source/Editor/Sequencer/Private/CurveEditor/Integration/Selection/CurveModelSyncer.cpp b/Engine/Source/Editor/Sequencer/Private/CurveEditor/Integration/Selection/CurveModelSyncer.cpp
--- a/Engine/Source/Editor/Sequencer/Private/CurveEditor/Integration/Selection/CurveModelSyncer.cpp
+++ b/Engine/Source/Editor/Sequencer/Private/CurveEditor/Integration/Selection/CurveModelSyncer.cpp
@@ -216,14 +216,10 @@ void FCurveModelSyncer::SyncSequencerToCurveEditor(FSequencer& InSequencer, FCur
 	// Guard against multiple broadcasts here and defer them until the end of this function
 	const FScopedCurveEditorTreeEventGuard ScopedEventGuard = InTree.ScopedEventGuard();
 	
-	// TODO UE-364008: In the future, we can remove the paths with bUseFilteredOut.
-	const bool bUseFilteredOut = InSequencer.GetSequencerSettings()->GetLinkFiltersWithCurveEditor();
-	const auto IsFilteredOut = [this, bUseFilteredOut](const TViewModelPtr<IOutlinerExtension>& InItem)
-		{
-			return UseLinkedFilteringUI()
-				? LinkedFilterViewModel->IsFilteredOut(InItem)
-				: bUseFilteredOut && InItem && InItem->IsFilteredOut();
-		};
+	const auto IsFilteredOut = [this](const TViewModelPtr<IOutlinerExtension>& InItem)
+	{
+		return LinkedFilterViewModel->IsFilteredOut(InItem);
+	};
 	
 	for (auto It = ViewModelToTreeItemIDMap.CreateIterator(); It; ++It)
 	{
diff --git a/Engine/Source/Editor/Sequencer/Private/CurveEditor/Views/CurveEditorWidgetOwner.cpp b/Engine/Source/Editor/Sequencer/Private/CurveEditor/Views/CurveEditorWidgetOwner.cpp
--- a/Engine/Source/Editor/Sequencer/Private/CurveEditor/Views/CurveEditorWidgetOwner.cpp
+++ b/Engine/Source/Editor/Sequencer/Private/CurveEditor/Views/CurveEditorWidgetOwner.cpp
@@ -40,8 +40,6 @@ static FCurveEditorWidgets MakePresenterWidgets(
 	const TSharedPtr<FTabManager> InTabManager
 	)
 {
-	const TSharedRef<SOldCurveTreeContent> OldContent = SNew(SOldCurveTreeContent, InArgs.Sequencer, InArgs.CurveEditor)
-		.TabManager(InTabManager);
 	const TSharedRef<SCurveTreeContent> TreeContent = SNew(SCurveTreeContent)
 		.TabManager(InTabManager)
 		.Sequencer(InArgs.Sequencer)
@@ -57,12 +55,6 @@ static FCurveEditorWidgets MakePresenterWidgets(
 			.LinkedFilterViewModel(InArgs.FilteringViewModel)
 		];
 	
-	const TSharedRef<SWidgetSwitcher> Content = 
-		SNew(SWidgetSwitcher)
-		.WidgetIndex_Lambda([]{ return UseLinkedFilteringUI() ? 1 : 0; })
-		+SWidgetSwitcher::Slot()[ OldContent ]
-		+SWidgetSwitcher::Slot()[ TreeContent ];
-	
 	const TSharedRef<SCurveEditorPanel> CurveEditorPanel = SNew(SCurveEditorPanel, InArgs.CurveEditor)
 		// Grid lines match the color specified in FSequencerTimeSliderController::OnPaintViewArea
 		.GridLineTint(FLinearColor(0.f, 0.f, 0.f, 0.3f))
@@ -72,17 +64,12 @@ static FCurveEditorWidgets MakePresenterWidgets(
 		.DisabledTimeSnapTooltip(LOCTEXT("CurveEditorTimeSnapDisabledTooltip", "Time Snapping is currently driven by Sequencer."))
 		.TreeContent()
 		[
-			Content
+			TreeContent
 		];
 	
 	const TSharedRef<SSequencerCurveEditor> RootWidget = SNew(SSequencerCurveEditor, CurveEditorPanel, InArgs.Sequencer);
-	return FCurveEditorWidgets(RootWidget, CurveEditorPanel, OldContent, TreeContent);
-}
+	return FCurveEditorWidgets(RootWidget, CurveEditorPanel, TreeContent);
 }
-
-bool UseLinkedFilteringUI()
-{
-	return PresenterDetail::CVarUseLinkedFilterUI.GetValueOnAnyThread();
 }
 
 FCurveEditorWidgetOwner::FCurveEditorWidgetOwner(const FCurveEditorWidgetOwnerArgs& InArgs)
@@ -150,21 +137,12 @@ void FCurveEditorWidgetOwner::CloseCurveEditor()
 
 void FCurveEditorWidgetOwner::SyncSelection() const
 {
-	if (UseLinkedFilteringUI())
-	{
-		Widgets.TreeContent->SyncSelection();
-	}
-	else
-	{
-		Widgets.OldTreeContent->SyncSelection();
-	}
+	Widgets.TreeContent->SyncSelection();
 }
 
 TSharedPtr<SCurveEditorTree> FCurveEditorWidgetOwner::GetTreeView() const
 {
-	return UseLinkedFilteringUI() 
-		? Widgets.TreeContent->GetCurveEditorTreeView()
-		: Widgets.OldTreeContent->GetCurveEditorTreeView();
+	return Widgets.TreeContent->GetCurveEditorTreeView();
 }
 
 USequencerSettings* FCurveEditorWidgetOwner::GetSequencerSettings() const
@@ -200,26 +178,12 @@ void FCurveEditorWidgetOwner::BindCommands()
 
 void FCurveEditorWidgetOwner::HandleQuickTreeSearch()
 {
-	if (UseLinkedFilteringUI())
-	{
-		Widgets.TreeContent->FocusSearchBox();
-	}
-	else
-	{
-		Widgets.OldTreeContent->FocusSearchBox();
-	}
+	Widgets.TreeContent->FocusSearchBox();
 }
 
 void FCurveEditorWidgetOwner::HandleToggleShowGotoBox()
 {
-	if (UseLinkedFilteringUI())
-	{
-		Widgets.TreeContent->FocusPlayTimeDisplay();
-	}
-	else
-	{
-		Widgets.OldTreeContent->FocusPlayTimeDisplay();
-	}
+	Widgets.TreeContent->FocusPlayTimeDisplay();
 }
 
 void FCurveEditorWidgetOwner::InitFrameNumberPropertyLayout(const TSharedRef<FSequencer>& InSequencer)
diff --git a/Engine/Source/Editor/Sequencer/Private/CurveEditor/Views/CurveEditorWidgetOwner.h b/Engine/Source/Editor/Sequencer/Private/CurveEditor/Views/CurveEditorWidgetOwner.h
--- a/Engine/Source/Editor/Sequencer/Private/CurveEditor/Views/CurveEditorWidgetOwner.h
+++ b/Engine/Source/Editor/Sequencer/Private/CurveEditor/Views/CurveEditorWidgetOwner.h
@@ -23,9 +23,6 @@ class FCurveModelSyncer;
 class FLinkedFilterViewModel;
 struct FCurveEditorWidgetOwnerArgs;
 
-/** @return Whether the linked filtering UI is enabled. */
-bool UseLinkedFilteringUI();
-
 /**
  * Instantiates the UI displaying Sequencer's FCurveEditor.
  * 
diff --git a/Engine/Source/Editor/Sequencer/Private/CurveEditor/Views/CurveEditorWidgets.h b/Engine/Source/Editor/Sequencer/Private/CurveEditor/Views/CurveEditorWidgets.h
--- a/Engine/Source/Editor/Sequencer/Private/CurveEditor/Views/CurveEditorWidgets.h
+++ b/Engine/Source/Editor/Sequencer/Private/CurveEditor/Views/CurveEditorWidgets.h
@@ -20,21 +20,16 @@ struct FCurveEditorWidgets
 	/** The curve editor panel. This is created and updated even if it is not currently visible. */
 	TSharedRef<SCurveEditorPanel> CurveEditorPanel;
 	
-	/** The tree content displayed inside CurveEditorPanel.*/
-	TSharedRef<SOldCurveTreeContent> OldTreeContent;
-	
 	/** The tree content for curve editor. It displays linked filtering UI. */
 	TSharedRef<SCurveTreeContent> TreeContent;
 
 	explicit FCurveEditorWidgets(
 		const TSharedRef<SSequencerCurveEditor>& InRootWidget, 
 		const TSharedRef<SCurveEditorPanel>& InCurveEditorPanel,
-		const TSharedRef<SOldCurveTreeContent>& InOldTreeContent,
 		const TSharedRef<SCurveTreeContent>& InTreeContent
 		)
 		: RootWidget(InRootWidget)
 		, CurveEditorPanel(InCurveEditorPanel)
-		, OldTreeContent(InOldTreeContent)
 		, TreeContent(InTreeContent)
 	{}
 };
diff --git a/Engine/Source/Editor/Sequencer/Private/CurveEditor/Widgets/SCurveEditorSearchAndFilterRow.cpp b/Engine/Source/Editor/Sequencer/Private/CurveEditor/Widgets/SCurveEditorSearchAndFilterRow.cpp
--- a/Engine/Source/Editor/Sequencer/Private/CurveEditor/Widgets/SCurveEditorSearchAndFilterRow.cpp
+++ b/Engine/Source/Editor/Sequencer/Private/CurveEditor/Widgets/SCurveEditorSearchAndFilterRow.cpp
@@ -98,10 +98,6 @@ TSharedRef<SWidget> SCurveEditorSearchAndFilterRow::MakeFilterContent(const FArg
 			.BeforeSearchBox()
 			[
 				SNew(SLinkedFilteringToggleButton, InArgs._FilterViewModel.ToSharedRef())
-				.Visibility_Lambda([]
-				{
-					return IsLinkedFilteringFeatureEnabled() ? EVisibility::Visible : EVisibility::Collapsed;
-				})
 			];
 	};
 	
diff --git a/Engine/Source/Editor/Sequencer/Private/Filters/Linking/Features/FilterAreaCommandBinder.cpp b/Engine/Source/Editor/Sequencer/Private/Filters/Linking/Features/FilterAreaCommandBinder.cpp
--- a/Engine/Source/Editor/Sequencer/Private/Filters/Linking/Features/FilterAreaCommandBinder.cpp
+++ b/Engine/Source/Editor/Sequencer/Private/Filters/Linking/Features/FilterAreaCommandBinder.cpp
@@ -77,7 +77,7 @@ void FFilterAreaCommandBinder::BindCommands()
 	RootCommandList->MapAction(
 		FSequencerTrackFilterCommands::Get().ToggleLinkedFiltering ,
 		FExecuteAction::CreateLambda([WeakViewModel = LinkedFilterModel.ToWeakPtr()] { ToggleLinkedFiltering(WeakViewModel); }),
-		FCanExecuteAction::CreateLambda([]{ return IsLinkedFilteringFeatureEnabled(); }),
+		FCanExecuteAction::CreateLambda([]{ return true; }),
 		FIsActionChecked::CreateLambda([WeakViewModel = LinkedFilterModel.ToWeakPtr()] { return IsLinkedFilteringActionChecked(WeakViewModel); })
 		);
 	
diff --git a/Engine/Source/Editor/Sequencer/Private/Filters/Menus/SequencerViewOptionsMenu.cpp b/Engine/Source/Editor/Sequencer/Private/Filters/Menus/SequencerViewOptionsMenu.cpp
--- a/Engine/Source/Editor/Sequencer/Private/Filters/Menus/SequencerViewOptionsMenu.cpp
+++ b/Engine/Source/Editor/Sequencer/Private/Filters/Menus/SequencerViewOptionsMenu.cpp
@@ -27,11 +27,6 @@ static void AddFilterBarSection(UToolMenu& InMenu)
 		);
 	
 	FilterBarSection.AddMenuEntry(TrackFilterCommands.ToggleFilterBarVisibility);
-	
-	if (IsLinkedFilteringFeatureEnabled() && UseLinkedFilteringUI())
-	{
-		FilterBarSection.AddMenuEntry(TrackFilterCommands.TogglePreserveFiltersOnUnlink);
-	}
 }
 
 static TSharedRef<SWidget> CreateMenu(FName InMenuName, FOnPopulateFilterBarMenu InPopulateMenu, const TSharedRef<FUICommandList>& InCommandList)
@@ -197,21 +192,6 @@ void FSequencerViewOptionsMenu::PopulateFilterOptionsSection(UToolMenu& InMenu)
 		),
 		EUserInterfaceActionType::ToggleButton);
 
-	// TODO UE-364008: Remove this once we decide to permanently support linked filtering in Curve Editor
-	if (!UseLinkedFilteringUI())
-	{
-		OptionsSection.AddMenuEntry(TEXT("LinkFiltersWithCurveEditor"),
-			LOCTEXT("LinkFiltersWithCurveEditor", "Link Filters with Curve Editor"),
-			LOCTEXT("LinkFiltersWithCurveEditorToolTip", "Toggle linking filters with the curve editor tree"),
-			FSlateIcon(),
-			FUIAction(
-				FExecuteAction::CreateRaw(this, &FSequencerViewOptionsMenu::ToggleLinkFiltersWithCurveEditor),
-				FCanExecuteAction(),
-				FIsActionChecked::CreateRaw(this, &FSequencerViewOptionsMenu::IsLinkFiltersWithCurveEditor)
-			),
-			EUserInterfaceActionType::ToggleButton);
-	}
-
 	ViewOptionsDetail::AddFilterBarSection(InMenu);
 }
 
@@ -295,42 +275,6 @@ void FSequencerViewOptionsMenu::ToggleAutoExpandPassedFilterNodes()
 	Sequencer->GetFilterInterface()->RequestFilterUpdate();
 }
 
-bool FSequencerViewOptionsMenu::IsLinkFiltersWithCurveEditor() const
-{
-	const TSharedPtr<FSequencer> Sequencer = WeakSequencer.Pin();
-	if (!Sequencer.IsValid())
-	{
-		return false;
-	}
-
-	const USequencerSettings* const SequencerSettings = Sequencer->GetSequencerSettings();
-	if (!SequencerSettings)
-	{
-		return false;
-	}
-
-	return SequencerSettings->GetLinkFiltersWithCurveEditor();
-}
-
-void FSequencerViewOptionsMenu::ToggleLinkFiltersWithCurveEditor()
-{
-	const TSharedPtr<FSequencer> Sequencer = WeakSequencer.Pin();
-	if (!Sequencer.IsValid())
-	{
-		return;
-	}
-
-	USequencerSettings* const SequencerSettings = Sequencer->GetSequencerSettings();
-	if (!SequencerSettings)
-	{
-		return;
-	}
-
-	SequencerSettings->SetLinkFiltersWithCurveEditor(!SequencerSettings->GetLinkFiltersWithCurveEditor());
-
-	Sequencer->GetFilterInterface()->RequestFilterUpdate();
-}
-
 TSharedPtr<SSequencer> FSequencerViewOptionsMenu::GetSequencerWidget() const
 {
 	const TSharedPtr<FSequencer> Sequencer = WeakSequencer.Pin();
diff --git a/Engine/Source/Editor/Sequencer/Private/Filters/Menus/SequencerViewOptionsMenu.h b/Engine/Source/Editor/Sequencer/Private/Filters/Menus/SequencerViewOptionsMenu.h
--- a/Engine/Source/Editor/Sequencer/Private/Filters/Menus/SequencerViewOptionsMenu.h
+++ b/Engine/Source/Editor/Sequencer/Private/Filters/Menus/SequencerViewOptionsMenu.h
@@ -41,9 +41,6 @@ protected:
 	bool IsAutoExpandPassedFilterNodes() const;
 	void ToggleAutoExpandPassedFilterNodes();
 
-	bool IsLinkFiltersWithCurveEditor() const;
-	void ToggleLinkFiltersWithCurveEditor();
-
 	TSharedPtr<SSequencer> GetSequencerWidget() const;
 
 	TWeakPtr<FSequencer> WeakSequencer;
diff --git a/Engine/Source/Editor/Sequencer/Private/Filters/Utils/LinkedFilteringViewUtils.cpp b/Engine/Source/Editor/Sequencer/Private/Filters/Utils/LinkedFilteringViewUtils.cpp
--- a/Engine/Source/Editor/Sequencer/Private/Filters/Utils/LinkedFilteringViewUtils.cpp
+++ b/Engine/Source/Editor/Sequencer/Private/Filters/Utils/LinkedFilteringViewUtils.cpp
@@ -30,14 +30,4 @@ bool IsLinkedFilteringActionChecked(TWeakPtr<ILinkedFilterViewModel> InWeakViewM
 	const TSharedPtr<ILinkedFilterViewModel> ViewModelPin = InWeakViewModel.Pin();
 	return ViewModelPin && ViewModelPin->GetFilterMode() == ELinkedFilterMode::Linked;
 }
-
-// TODO UE-371833: Remove CVar in 5.9 assuming the system is stable.
-TAutoConsoleVariable<bool> CVarSequencerUseLinkedFiltering(
-	TEXT("Sequencer.UseLinkedFiltering"), true, TEXT("Whether to use the new linked filtering system")
-	);
-
-bool IsLinkedFilteringFeatureEnabled()
-{
-	return CVarSequencerUseLinkedFiltering.GetValueOnAnyThread();
-}
 } // namespace UE::Sequencer
diff --git a/Engine/Source/Editor/Sequencer/Private/Filters/Utils/LinkedFilteringViewUtils.h b/Engine/Source/Editor/Sequencer/Private/Filters/Utils/LinkedFilteringViewUtils.h
--- a/Engine/Source/Editor/Sequencer/Private/Filters/Utils/LinkedFilteringViewUtils.h
+++ b/Engine/Source/Editor/Sequencer/Private/Filters/Utils/LinkedFilteringViewUtils.h
@@ -12,7 +12,4 @@ class ILinkedFilterViewModel;
 void ToggleLinkedFiltering(TWeakPtr<ILinkedFilterViewModel> InWeakViewModel);
 /** Utility for binding ToggleLinkedFiltering to FIsActionChecked. */
 bool IsLinkedFilteringActionChecked(TWeakPtr<ILinkedFilterViewModel> InWeakViewModel);
-
-/** @return Whether Sequencer should use linked filtering. */
-bool IsLinkedFilteringFeatureEnabled();
 }
diff --git a/Engine/Source/Editor/Sequencer/Private/SSequencer.cpp b/Engine/Source/Editor/Sequencer/Private/SSequencer.cpp
--- a/Engine/Source/Editor/Sequencer/Private/SSequencer.cpp
+++ b/Engine/Source/Editor/Sequencer/Private/SSequencer.cpp
@@ -390,15 +390,13 @@ void SSequencer::Construct(const FArguments& InArgs, TSharedRef<FSequencer> InSe
 	const TSharedRef<FSequencerEditorViewModelImpl> RootViewModel = SequencerRef->GetViewModelImpl();
 	const auto IsLinkedToggleButtonVisible = [WeakRootViewModel = RootViewModel.ToWeakPtr()]()
 	{
-		const bool bIsFeatureEnabled = IsLinkedFilteringFeatureEnabled();
-			
 		// The linked filtering toggle button should only be displayed if there is another UI system that can link... if not, the feature is pointless
 		const TSharedPtr<FSequencerEditorViewModelImpl> RootViewModel = WeakRootViewModel.Pin();
 		const TSharedPtr<FLinkedFilterFactoryViewModel> Factory = RootViewModel ? RootViewModel->GetLinkedFilterViewModelFactoryImpl() : nullptr;
 		// Check for more than 1 because Sequencer's instance counts as a user (@see FSequencerEditorViewModelImpl::SequenceFilterInstanceViewModel)
 		const bool bIsAnySystemUsingLinkedFiltering = Factory && Factory->GetFilterFeatureUseCount() > 1;
 			
-		const bool bShouldBeVisible = bIsFeatureEnabled && bIsAnySystemUsingLinkedFiltering;
+		const bool bShouldBeVisible = bIsAnySystemUsingLinkedFiltering;
 		return bShouldBeVisible ? EVisibility::Visible : EVisibility::Collapsed;
 	};
 	
diff --git a/Engine/Source/Editor/Sequencer/Private/SequencerSettings.cpp b/Engine/Source/Editor/Sequencer/Private/SequencerSettings.cpp
--- a/Engine/Source/Editor/Sequencer/Private/SequencerSettings.cpp
+++ b/Engine/Source/Editor/Sequencer/Private/SequencerSettings.cpp
@@ -38,7 +38,6 @@ USequencerSettings::USequencerSettings( const FObjectInitializer& ObjectInitiali
 	ZoomPosition = ESequencerZoomPosition::SZP_CurrentTime;
 	bAutoScrollEnabled = false;
 	bLinkCurveEditorTimeRange = false;
-	bLinkFiltersWithCurveEditor = false;
 	bSynchronizeCurveEditorSelection = true;
 	bSynchronizeOutlinerSelectionToCurveEditor = true;
 	bIsolateCurveEditorToSelection = true;
@@ -575,19 +574,23 @@ void USequencerSettings::SetLinkCurveEditorTimeRange(bool InbLinkCurveEditorTime
 	}
 }
 
+PRAGMA_DISABLE_DEPRECATION_WARNINGS
 bool USequencerSettings::GetLinkFiltersWithCurveEditor() const
 {
-	return bLinkFiltersWithCurveEditor;
+	return bLinkFiltersWithCurveEditor_DEPRECATED;
 }
+PRAGMA_ENABLE_DEPRECATION_WARNINGS
 
+PRAGMA_DISABLE_DEPRECATION_WARNINGS
 void USequencerSettings::SetLinkFiltersWithCurveEditor(bool bInLinkFiltersWithCurveEditor)
 {
-	if (bLinkFiltersWithCurveEditor != bInLinkFiltersWithCurveEditor)
+	if (bLinkFiltersWithCurveEditor_DEPRECATED != bInLinkFiltersWithCurveEditor)
 	{
-		bLinkFiltersWithCurveEditor = bInLinkFiltersWithCurveEditor;
+		bLinkFiltersWithCurveEditor_DEPRECATED = bInLinkFiltersWithCurveEditor;
 		SaveConfig();
 	}
 }
+PRAGMA_ENABLE_DEPRECATION_WARNINGS
 
 void USequencerSettings::SyncCurveEditorSelection(bool bInSynchronizeCurveEditorSelection)
 {
diff --git a/Engine/Source/Editor/Sequencer/Public/SequencerSettings.h b/Engine/Source/Editor/Sequencer/Public/SequencerSettings.h
--- a/Engine/Source/Editor/Sequencer/Public/SequencerSettings.h
+++ b/Engine/Source/Editor/Sequencer/Public/SequencerSettings.h
@@ -294,8 +294,10 @@ public:
 	UE_API void SetLinkCurveEditorTimeRange(bool InbLinkCurveEditorTimeRange);
 
 	/** Gets whether or not to link sequence filters with the curve editor tree */
+	UE_DEPRECATED(5.8, "Linked filtering replaces this feature")
 	UE_API bool GetLinkFiltersWithCurveEditor() const;
 	/** Sets whether or not to link sequence filters with the curve editor tree */
+	UE_DEPRECATED(5.8, "Linked filtering replaces this feature")
 	UE_API void SetLinkFiltersWithCurveEditor(bool bInLinkFiltersWithCurveEditor);
 
 	/** Return true whether changing the sequencer tree selection will also select the relevant nodes in the curve editor tree if possible */
@@ -759,8 +761,9 @@ protected:
 	bool bLinkCurveEditorTimeRange;
 
 	/** Enable or disable linking sequence filters with the curve editor tree */
-	UPROPERTY(config, EditAnywhere, Category = CurveEditor)
-	bool bLinkFiltersWithCurveEditor;
+	UE_DEPRECATED(5.8, "Linked filtering replaces this feature")
+	UPROPERTY(config)
+	bool bLinkFiltersWithCurveEditor_DEPRECATED;
 
 	/** When enabled, changing the sequencer tree selection will also select the relevant nodes in the curve editor tree if possible. */
 	UPROPERTY(config, EditAnywhere, Category=CurveEditor, meta = (DisplayName = "Synchronize Curve Editor Selection to Outliner"))
