Sequencer: Using "Hide Selected Tracks" filter in sequencer or sequencer curves hide all elements if nothing is selected

#Synced-CL 52617198

diff --git a/Engine/Source/Editor/Sequencer/Private/Filters/ViewModels/HideIsolateShowViewModel.cpp b/Engine/Source/Editor/Sequencer/Private/Filters/ViewModels/HideIsolateShowViewModel.cpp
--- a/Engine/Source/Editor/Sequencer/Private/Filters/ViewModels/HideIsolateShowViewModel.cpp
+++ b/Engine/Source/Editor/Sequencer/Private/Filters/ViewModels/HideIsolateShowViewModel.cpp
@@ -33,7 +33,7 @@ void FHideIsolateShowViewModel::EmptyIsolatedTracks()
 
 void FHideIsolateShowViewModel::IsolateSelectedTracks(bool bAddToExisting)
 {
-	IsolateFilter->IsolateTracks(GetSelectedTracksOrAll(), bAddToExisting);
+	IsolateFilter->IsolateTracks(GetSelectedTracks(), bAddToExisting);
 }
 
 void FHideIsolateShowViewModel::EmptyHiddenTracks()
@@ -44,7 +44,7 @@ void FHideIsolateShowViewModel::EmptyHiddenTracks()
 
 void FHideIsolateShowViewModel::HideSelectedTracks(const bool bAddToExisting)
 {
-	IsolateFilter->HideTracks(GetSelectedTracksOrAll(), bAddToExisting);
+	IsolateFilter->HideTracks(GetSelectedTracks(), bAddToExisting);
 }
 
 void FHideIsolateShowViewModel::ShowAllTracks()
@@ -151,14 +151,6 @@ FText FHideIsolateShowViewModel::GetShowAllTracksButtonTooltipText() const
 	return TooltipText;
 }
 
-TSet<TWeakViewModelPtr<IOutlinerExtension>> FHideIsolateShowViewModel::GetSelectedTracksOrAll() const
-{
-	const TSet<TWeakViewModelPtr<IOutlinerExtension>> Selection = SelectedTracksAttr.Get();
-	return Selection.IsEmpty()
-		? AllTracksAttr.Get()
-		: Selection;
-}
-
 TSet<TWeakViewModelPtr<IOutlinerExtension>> FHideIsolateShowViewModel::GetSelectedTracks() const
 {
 	const TSet<TWeakViewModelPtr<IOutlinerExtension>> Selection = SelectedTracksAttr.Get();
diff --git a/Engine/Source/Editor/Sequencer/Private/Filters/ViewModels/HideIsolateShowViewModel.h b/Engine/Source/Editor/Sequencer/Private/Filters/ViewModels/HideIsolateShowViewModel.h
--- a/Engine/Source/Editor/Sequencer/Private/Filters/ViewModels/HideIsolateShowViewModel.h
+++ b/Engine/Source/Editor/Sequencer/Private/Filters/ViewModels/HideIsolateShowViewModel.h
@@ -77,9 +77,6 @@ private:
 	
 	/** Broadcasts when an item is supposed to be focused, e.g. scrolled to view. */
 	FRequestScrollToViewEvent RequestScrollToViewDelegate;
-	
-	/** @return Selected tracks, or all of them. */
-	TSet<TWeakViewModelPtr<IOutlinerExtension>> GetSelectedTracksOrAll() const;
 
 	/** @return Selected tracks. */
 	TSet<TWeakViewModelPtr<IOutlinerExtension>> GetSelectedTracks() const;
