Sequencer: Config: After restart, UI should go to the linked or instanced filter mode it last had

#Synced-CL 52617198

diff --git a/Engine/Source/Editor/Sequencer/Private/Filters/Linking/Features/FilterModeConfigManager.cpp b/Engine/Source/Editor/Sequencer/Private/Filters/Linking/Features/FilterModeConfigManager.cpp
new file mode 100644
--- /dev/null
+++ b/Engine/Source/Editor/Sequencer/Private/Filters/Linking/Features/FilterModeConfigManager.cpp
@@ -0,0 +1,56 @@
+﻿// Copyright Epic Games, Inc. All Rights Reserved.
+
+#include "FilterModeConfigManager.h"
+
+#include "Sequencer.h"
+#include "Filters/Linking/Mode/ILinkedFilterViewModel.h"
+
+namespace UE::Sequencer
+{
+namespace ConfigDetail
+{
+static void ApplyConfig(
+	FName InFilterAreaConfigId, const TWeakPtr<FSequencer>& InWeakSequencer, const TSharedRef<ILinkedFilterViewModel>& InLinkedFilterViewModel
+	)
+{
+	const TSharedPtr<FSequencer> Sequencer = InWeakSequencer.Pin();
+	USequencerSettings* Settings = Sequencer ? Sequencer->GetSequencerSettings() : nullptr;
+	FSequencerFilterAreaConfig* Config = Settings ? Settings->FindTrackFilterArea(InFilterAreaConfigId) : nullptr;
+	if (Config)
+	{
+		InLinkedFilterViewModel->SetFilterMode(Config->IsLinkedFiltering() ? ELinkedFilterMode::Linked : ELinkedFilterMode::Instanced); 
+	}
+}
+}
+
+FFilterModeConfigManager::FFilterModeConfigManager(
+	FName InFilterAreaConfigId, TWeakPtr<FSequencer> InWeakSequencer, const TSharedRef<ILinkedFilterViewModel>& InLinkedFilterViewModel
+	)
+	: FilterAreaConfigId(InFilterAreaConfigId)
+	, WeakSequencer(MoveTemp(InWeakSequencer))
+	, LinkingModel(InLinkedFilterViewModel)
+{
+	ConfigDetail::ApplyConfig(FilterAreaConfigId, WeakSequencer, LinkingModel);
+	LinkingModel->OnFilterModeChanged().AddRaw(this, &FFilterModeConfigManager::OnFilterModeChanged);
+}
+
+FFilterModeConfigManager::~FFilterModeConfigManager()
+{
+	LinkingModel->OnFilterModeChanged().RemoveAll(this);
+}
+
+void FFilterModeConfigManager::OnFilterModeChanged() const
+{
+	const TSharedPtr<FSequencer> Sequencer = WeakSequencer.Pin();
+	USequencerSettings* Settings = Sequencer ? Sequencer->GetSequencerSettings() : nullptr;
+	if (!Settings)
+	{
+		return;
+	}
+	
+	constexpr bool bSaveConfig = false;
+	FSequencerFilterAreaConfig& Config = Settings->FindOrAddTrackFilterArea(FilterAreaConfigId, bSaveConfig);
+	Config.SetIsLinkedFiltering(LinkingModel->GetFilterMode() == ELinkedFilterMode::Linked);
+	Settings->SaveConfig();
+}
+} // namespace UE::Sequencer
diff --git a/Engine/Source/Editor/Sequencer/Private/Filters/Linking/Features/FilterModeConfigManager.h b/Engine/Source/Editor/Sequencer/Private/Filters/Linking/Features/FilterModeConfigManager.h
new file mode 100644
--- /dev/null
+++ b/Engine/Source/Editor/Sequencer/Private/Filters/Linking/Features/FilterModeConfigManager.h
@@ -0,0 +1,37 @@
+﻿// Copyright Epic Games, Inc. All Rights Reserved.
+
+#pragma once
+
+#include "Templates/SharedPointer.h"
+#include "Templates/UnrealTemplate.h"
+
+class FSequencer;
+
+namespace UE::Sequencer
+{
+class ILinkedFilterViewModel;
+
+/** Responsible for saving and loading FSequencerFilterAreaConfig::bIsLinkedFiltering. */
+class FFilterModeConfigManager : public FNoncopyable
+{
+public:
+	
+	explicit FFilterModeConfigManager(
+		FName InFilterAreaConfigId, TWeakPtr<FSequencer> InWeakSequencer, const TSharedRef<ILinkedFilterViewModel>& InLinkedFilterViewModel
+		);
+	~FFilterModeConfigManager();
+	
+private:
+	
+	/** @see USequencerSettings::FindOrAddTrackFilterArea. */
+	const FName FilterAreaConfigId;
+	/** Used to get Sequencer settings. */
+	const TWeakPtr<FSequencer> WeakSequencer;
+	
+	/** Used to listen when the filter mode changes */
+	const TSharedRef<ILinkedFilterViewModel> LinkingModel;
+	
+	/** Saves the filter mode to config. */
+	void OnFilterModeChanged() const;
+};
+} // namespace UE::Sequencer
diff --git a/Engine/Source/Editor/Sequencer/Private/Filters/Linking/FilterAreaManager.cpp b/Engine/Source/Editor/Sequencer/Private/Filters/Linking/FilterAreaManager.cpp
--- a/Engine/Source/Editor/Sequencer/Private/Filters/Linking/FilterAreaManager.cpp
+++ b/Engine/Source/Editor/Sequencer/Private/Filters/Linking/FilterAreaManager.cpp
@@ -3,7 +3,7 @@
 #include "FilterAreaManager.h"
 
 #include "Sequencer.h"
-#include "Filters/SequencerTrackFilterCommands.h"
+#include "Filters/Linking/Mode/LinkedFilterViewModel.h"
 
 namespace UE::Sequencer
 {
@@ -15,5 +15,6 @@ FFilterAreaManager::FFilterAreaManager(
 	: FilterModel(InFilterModel)
 	, CommandBinder(InFilterAreaConfigId, InWeakSequencer, FilterModel)
 	, LinkedToUnlinkedCopier(FilterModel, InFilterAreaConfigId, InWeakSequencer)
+	, FilterModeConfigManager(InFilterAreaConfigId, InWeakSequencer, FilterModel)
 {}
 }
diff --git a/Engine/Source/Editor/Sequencer/Private/Filters/Linking/FilterAreaManager.h b/Engine/Source/Editor/Sequencer/Private/Filters/Linking/FilterAreaManager.h
--- a/Engine/Source/Editor/Sequencer/Private/Filters/Linking/FilterAreaManager.h
+++ b/Engine/Source/Editor/Sequencer/Private/Filters/Linking/FilterAreaManager.h
@@ -3,6 +3,7 @@
 #pragma once
 
 #include "Features/FilterAreaCommandBinder.h"
+#include "Features/FilterModeConfigManager.h"
 #include "Features/LinkedFilterStatePreserver.h"
 #include "Framework/Commands/UICommandList.h"
 #include "Templates/SharedPointer.h"
@@ -43,10 +44,13 @@ private:
 	const TSharedRef<FLinkedFilterViewModel> FilterModel;
 	
 	/** Handles the commands that operate on the filter area on a whole. */
-	FFilterAreaCommandBinder CommandBinder;
+	const FFilterAreaCommandBinder CommandBinder;
 	
 	/** Copies the filter state from the linked to the unlinked filter bar on switch. */
-	FLinkedFilterStatePreserver LinkedToUnlinkedCopier;
+	const FLinkedFilterStatePreserver LinkedToUnlinkedCopier;
+	
+	/** Saves and loads the filter mode. Example: if you had UI in instanced mode, when you restart the editor, it will be in instanced mode again. */
+	const FFilterModeConfigManager FilterModeConfigManager;
 };
 }
 
diff --git a/Engine/Source/Editor/Sequencer/Private/Filters/SequencerFilterAreaConfig.cpp b/Engine/Source/Editor/Sequencer/Private/Filters/SequencerFilterAreaConfig.cpp
--- a/Engine/Source/Editor/Sequencer/Private/Filters/SequencerFilterAreaConfig.cpp
+++ b/Engine/Source/Editor/Sequencer/Private/Filters/SequencerFilterAreaConfig.cpp
@@ -21,3 +21,13 @@ bool FSequencerFilterAreaConfig::GetPreserveFiltersOnUnlink() const
 {
 	return bPreserveFiltersOnUnlink;
 }
+
+void FSequencerFilterAreaConfig::SetIsLinkedFiltering(bool bValue)
+{
+	bIsLinkedFiltering = bValue;
+}
+
+bool FSequencerFilterAreaConfig::IsLinkedFiltering() const
+{
+	return bIsLinkedFiltering;
+}
diff --git a/Engine/Source/Editor/Sequencer/Public/Filters/SequencerFilterAreaConfig.h b/Engine/Source/Editor/Sequencer/Public/Filters/SequencerFilterAreaConfig.h
--- a/Engine/Source/Editor/Sequencer/Public/Filters/SequencerFilterAreaConfig.h
+++ b/Engine/Source/Editor/Sequencer/Public/Filters/SequencerFilterAreaConfig.h
@@ -21,6 +21,9 @@ struct FSequencerFilterAreaConfig
 	UE_API void SetPreserveFiltersOnUnlink(bool bValue);
 	UE_API bool GetPreserveFiltersOnUnlink() const;
 	
+	UE_API void SetIsLinkedFiltering(bool bValue);
+	UE_API bool IsLinkedFiltering() const;
+	
 private:
 	
 	/** Whether the filter bar should be visible. */
@@ -31,7 +34,9 @@ private:
 	UPROPERTY()
 	bool bPreserveFiltersOnUnlink = true;
 	
-	// TODO UE-362150: safe whether the filter bar should be instanced or linked
+	/** Whether the UI is in linked filtering mode. True: linked filter. False: Instanced filtering.  */
+	UPROPERTY()
+	bool bIsLinkedFiltering = true;
 };
 
 #undef UE_API
