Sequencer: Delete unused code. This was missed in a previous submit.

#Synced-CL 52617198

diff --git a/Engine/Source/Editor/Sequencer/Private/CurveEditor/Views/CurveEditorWidgetOwner.cpp b/Engine/Source/Editor/Sequencer/Private/CurveEditor/Views/CurveEditorWidgetOwner.cpp
--- a/Engine/Source/Editor/Sequencer/Private/CurveEditor/Views/CurveEditorWidgetOwner.cpp
+++ b/Engine/Source/Editor/Sequencer/Private/CurveEditor/Views/CurveEditorWidgetOwner.cpp
@@ -4,8 +4,7 @@
 
 #include "CurveEditorWidgetOwnerArgs.h"
 #include "CurveEditor/Widgets/SSequencerCurveEditor.h"
-#include "CurveEditor/Widgets/Tree/SCurveTreeContent.h"
-#include "CurveEditor/Widgets/Tree/SOldCurveTreeContent.h"
+#include "CurveEditor/Widgets/SCurveTreeContent.h"
 #include "Filters/SCurveEditorFilterPanel.h"
 #include "Filters/Linking/Mode/LinkedFilterViewModel.h"
 #include "Filters/Widgets/Linking/SFilterBarSwitcher.h"
@@ -30,11 +29,6 @@ namespace UE::Sequencer
 {
 namespace PresenterDetail
 {
-// TODO UE-371833: Remove CVar in 5.9 assuming the system is stable.
-static TAutoConsoleVariable<bool> CVarUseLinkedFilterUI(
-	TEXT("Sequencer.CurveEditor.UseLinkedFiltering"), true, TEXT("Whether to use the new linked filtering UI in Curve Editor")
-	);
-
 static FCurveEditorWidgets MakePresenterWidgets(
 	const FCurveEditorWidgetOwnerArgs& InArgs,
 	const TSharedPtr<FTabManager> InTabManager
diff --git a/Engine/Source/Editor/Sequencer/Private/CurveEditor/Widgets/SCurveTreeContent.cpp b/Engine/Source/Editor/Sequencer/Private/CurveEditor/Widgets/SCurveTreeContent.cpp
new file mode 100644
--- /dev/null
+++ b/Engine/Source/Editor/Sequencer/Private/CurveEditor/Widgets/SCurveTreeContent.cpp
@@ -0,0 +1,248 @@
+﻿// Copyright Epic Games, Inc. All Rights Reserved.
+
+#include "SCurveTreeContent.h"
+
+#include "Components/VerticalBox.h"
+#include "CurveEditor.h"
+#include "CurveEditor/Integration/Selection/CurveModelSyncer.h"
+#include "CurveEditor/Views/SelectionSyncUtils.h"
+#include "CurveEditor/Widgets/SCurveEditorSearchAndFilterRow.h"
+#include "Filters/Widgets/SSequencerSearchBox.h"
+#include "Misc/TransportControlsHelper.h"
+#include "SPopoutTabInlineContent.h"
+#include "STemporarilyFocusedSpinBox.h"
+#include "Sequencer.h"
+#include "Styling/AppStyle.h"
+#include "Tree/SCurveEditorTree.h"
+#include "Widgets/Layout/SBorder.h"
+#include "Widgets/Layout/SScrollBorder.h"
+
+namespace UE::Sequencer
+{
+void SCurveTreeContent::Construct(const FArguments& InArgs)
+{
+	const bool bHasCurveModelSyncer = InArgs._CurveModelSyncer.IsBound() || InArgs._CurveModelSyncer.IsSet();
+	check(InArgs._Sequencer && InArgs._CurveEditor && InArgs._FilterViewModel && bHasCurveModelSyncer);
+	
+	WeakSequencer = InArgs._Sequencer;
+	CurveEditor = InArgs._CurveEditor;
+	CurveModelSyncerAttr = InArgs._CurveModelSyncer;
+	
+	ChildSlot
+	[
+		SNew(SVerticalBox)
+
+		+ SVerticalBox::Slot()
+		.AutoHeight()
+		[
+			MakeFilterArea(InArgs)
+		]
+
+		+SVerticalBox::Slot()
+		[
+			MakeTreeContent(InArgs)
+		]
+
+		+ SVerticalBox::Slot()
+		.AutoHeight()
+		[
+			MakeFooter(InArgs)
+		]
+	];
+}
+
+void SCurveTreeContent::FocusSearchBox()
+{
+	FSlateApplication::Get().SetKeyboardFocus(FilterRow->GetActiveSearchBox(), EFocusCause::SetDirectly);
+}
+
+void SCurveTreeContent::FocusPlayTimeDisplay()
+{
+	PlayTimeDisplay->Setup(); 
+	FSlateApplication::Get().SetKeyboardFocus(PlayTimeDisplay, EFocusCause::SetDirectly);
+}
+
+void SCurveTreeContent::SyncSelection() const
+{
+	const TSharedPtr<FSequencer> Sequencer = WeakSequencer.Pin();
+	if (!Sequencer)
+	{
+		return;
+	}
+	
+	ApplySequencerToCurveEditorSelection(
+		Sequencer->GetViewModel()->GetSelection(), CurveEditor.ToSharedRef(), CurveEditorTreeView.ToSharedRef()
+		);
+	
+	if (Sequencer->GetSequencerSettings()->GetAutoExpandCurveEditorOnSelection())
+	{
+		ExpandSelectedCurveEditorItems(*CurveEditor->GetTree(), *CurveEditorTreeView.ToSharedRef());
+	}
+}
+
+TSharedRef<SWidget> SCurveTreeContent::MakeFilterArea(const FArguments& InArgs)
+{
+	return SNew(SBorder)
+		.BorderImage(FAppStyle::GetBrush(TEXT("ToolPanel.GroupBorder")))
+		.Padding(FAppStyle::GetMargin(TEXT("Sequencer.FilterAreaMargin")))
+		.Clipping(EWidgetClipping::ClipToBounds)
+		[
+			SNew(SVerticalBox)
+		
+			+SVerticalBox::Slot()
+			.AutoHeight()
+			[
+				// The scroll box adds a shadow to the right when the toolbar is narrow and hiding buttons.
+				SNew(SScrollBox)
+				.Orientation(Orient_Horizontal)
+				.ScrollBarAlwaysVisible(false)
+				.ScrollBarVisibility(EVisibility::Collapsed)
+				
+				+ SScrollBox::Slot()
+				.HAlign(HAlign_Fill)
+				.FillContentSize(1.f)
+				[
+					SAssignNew(FilterRow, SCurveEditorSearchAndFilterRow)
+					.Sequencer(InArgs._Sequencer)
+					.FilterViewModel(InArgs._FilterViewModel)
+					.CommandList(InArgs._CommandList)
+					.ScrollItemIntoView(this, &SCurveTreeContent::ScrollIntoView)
+					.AllItems(this, &SCurveTreeContent::GetAllItems)
+					.SelectedItems(this, &SCurveTreeContent::GetSelectedItems)
+				]
+			]
+			
+			+SVerticalBox::Slot()
+			.AutoHeight()
+			[
+				InArgs._FilterPills.Widget
+			]
+		];
+}
+
+TSharedRef<SWidget> SCurveTreeContent::MakeTreeContent(const FArguments& InArgs)
+{
+	const TSharedRef<SCurveEditorTree> NewCurveEditorTree = SNew(SCurveEditorTree, InArgs._CurveEditor);
+	CurveEditorTreeView = NewCurveEditorTree;
+	return SNew(SScrollBorder, NewCurveEditorTree)
+		[
+			NewCurveEditorTree
+		];
+}
+
+TSharedRef<SWidget> SCurveTreeContent::MakeFooter(const FArguments& InArgs)
+{
+	const TSharedRef<FSequencer> Sequencer = InArgs._Sequencer.ToSharedRef();
+	
+	PlayTimeDisplay = Sequencer->MakePlayTimeDisplay();
+	return SNew(SBorder)
+		.BorderImage(FAppStyle::GetBrush("ToolPanel.GroupBorder"))
+		.Clipping(EWidgetClipping::ClipToBounds)
+		.Visibility(this, &SCurveTreeContent::GetTransportControlsVisibility)
+		[
+			SNew(SHorizontalBox)
+
+			+ SHorizontalBox::Slot()
+			.HAlign(HAlign_Fill)
+			[
+				MakePopoutTransportControls(Sequencer, InArgs._TabManager)
+			]
+
+			+ SHorizontalBox::Slot()
+			.AutoWidth()
+			.VAlign(VAlign_Center)
+			.HAlign(HAlign_Right)
+			[
+				SNew(SButton)
+				.VAlign(EVerticalAlignment::VAlign_Center)
+				.ButtonStyle(FAppStyle::Get(), "NoBorder")
+				.ContentPadding(FMargin(1, 0))
+				[
+					SNew(SHorizontalBox)
+
+					+ SHorizontalBox::Slot()
+					.AutoWidth()
+					.VAlign(VAlign_Center)
+					.HAlign(HAlign_Right)
+					.Padding(FMargin(3.f, 0.f, 0.f, 0.f))
+					[
+						SNew(SBorder)
+						.BorderImage(nullptr)
+						[
+							PlayTimeDisplay.ToSharedRef()
+						]
+					]
+				]
+			]
+		];
+}
+
+void SCurveTreeContent::ScrollIntoView(const TWeakViewModelPtr<IOutlinerExtension>& WeakViewModel)
+{
+	const FCurveModelSyncer* CurveModelSyncer = CurveModelSyncerAttr.Get();
+	if (!CurveModelSyncer)
+	{
+		return;
+	}
+	
+	const FCurveEditorTreeItemID Id = CurveModelSyncer->FindIdByViewModel(WeakViewModel);
+	if (Id.IsValid())
+	{
+		CurveEditorTreeView->RequestScrollIntoView(Id);
+	}
+}
+
+TSet<TWeakViewModelPtr<IOutlinerExtension>> SCurveTreeContent::GetAllItems() const
+{
+	const FCurveModelSyncer* CurveModelSyncer = CurveModelSyncerAttr.Get();
+	if (!CurveModelSyncer)
+	{
+		return {};
+	}
+	
+	TSet<TWeakViewModelPtr<IOutlinerExtension>> AllItems;
+	for (const TPair<FCurveEditorTreeItemID, FCurveEditorTreeItem>& Item : CurveEditor->GetTree()->GetAllItems())
+	{
+		const TSharedPtr<FViewModel> ViewModel = CurveModelSyncer->FindViewModelById(Item.Key).Pin();
+		const TViewModelPtr<IOutlinerExtension> OutlinerNode = ViewModel ? CastViewModel<IOutlinerExtension>(ViewModel.ToSharedRef()) : nullptr;
+		if (OutlinerNode)
+		{
+			AllItems.Add(OutlinerNode);
+		}
+	}
+	return AllItems;
+}
+
+TSet<TWeakViewModelPtr<IOutlinerExtension>> SCurveTreeContent::GetSelectedItems() const
+{
+	const FCurveModelSyncer* CurveModelSyncer = CurveModelSyncerAttr.Get();
+	if (!CurveModelSyncer)
+	{
+		return {};
+	}
+	
+	TSet<TWeakViewModelPtr<IOutlinerExtension>> Selected;
+	for (const FCurveEditorTreeItemID& Item : CurveEditorTreeView->GetSelectedItems())
+	{
+		const TSharedPtr<FViewModel> ViewModel = CurveModelSyncer->FindViewModelById(Item).Pin();
+		const TViewModelPtr<IOutlinerExtension> OutlinerNode = ViewModel ? CastViewModel<IOutlinerExtension>(ViewModel.ToSharedRef()) : nullptr;
+		if (OutlinerNode)
+		{
+			Selected.Add(OutlinerNode);
+		}
+	}
+	return Selected;
+}
+
+EVisibility SCurveTreeContent::GetTransportControlsVisibility() const
+{
+	if (const TSharedPtr<FSequencer> Sequencer = WeakSequencer.Pin())
+	{
+		const bool bShowTransportControls = GetDefault<UCurveEditorSettings>()->GetShowTransportControls();
+		return IsPopoutTransportControlsVisible(Sequencer.ToSharedRef(), bShowTransportControls)
+			? EVisibility::Visible : EVisibility::Collapsed;
+	}
+	return EVisibility::Visible;
+}
+
+} // namespace UE::Sequencer
diff --git a/Engine/Source/Editor/Sequencer/Private/CurveEditor/Widgets/SCurveTreeContent.h b/Engine/Source/Editor/Sequencer/Private/CurveEditor/Widgets/SCurveTreeContent.h
new file mode 100644
--- /dev/null
+++ b/Engine/Source/Editor/Sequencer/Private/CurveEditor/Widgets/SCurveTreeContent.h
@@ -0,0 +1,89 @@
+﻿// Copyright Epic Games, Inc. All Rights Reserved.
+
+#pragma once
+
+#include "MVVM/Extensions/IOutlinerExtension.h"
+#include "Widgets/DeclarativeSyntaxSupport.h"
+#include "Widgets/SCompoundWidget.h"
+
+class FCurveEditor;
+class FSequencer;
+class FTabManager;
+class FUICommandList;
+class SCurveEditorTree;
+template<typename T> struct STemporarilyFocusedSpinBox;
+
+namespace UE::Sequencer
+{
+class FCurveModelSyncer;
+class FLinkedFilterViewModel;
+class SCurveEditorSearchAndFilterRow;
+
+/** Displays Curve Editor in Sequencer with linked filtering. */
+class SCurveTreeContent : public SCompoundWidget
+{
+public:
+	
+	SLATE_BEGIN_ARGS(SCurveTreeContent) {}
+		/** Required. The Sequencer instance that owns this content */
+		SLATE_ARGUMENT(TSharedPtr<FSequencer>, Sequencer)
+		/** Required. The curve editor instance that is being displayed */
+		SLATE_ARGUMENT(TSharedPtr<FCurveEditor>, CurveEditor)
+		/** Required. Determines the filtering widgets displayed. */
+		SLATE_ARGUMENT(TSharedPtr<FLinkedFilterViewModel>, FilterViewModel)
+		/** Required. The command list that should be used, e.g. by the view options context menu. */
+		SLATE_ARGUMENT(TSharedPtr<FUICommandList>, CommandList)
+		
+		/** Explicit tab manager to use, otherwise uses the global tab manager */
+		SLATE_ARGUMENT(TSharedPtr<FTabManager>, TabManager)
+
+		/** Used to translate between IOutlinerExtension and FCurveEditorTreeItemID. */
+		SLATE_ATTRIBUTE(const FCurveModelSyncer*, CurveModelSyncer)
+
+		/** Shows the filter pills. */
+		SLATE_NAMED_SLOT(FArguments, FilterPills)
+	SLATE_END_ARGS()
+	
+	void Construct(const FArguments& InArgs);
+	
+	/** Sets keyboard focus to the search box. */
+	void FocusSearchBox();
+	/** Sets keyboard focus to the play time display. */
+	void FocusPlayTimeDisplay();
+	
+	/** Syncs the tree view's selection to that of Sequencer. */
+	void SyncSelection() const;
+	
+	const TSharedPtr<SCurveEditorTree>& GetCurveEditorTreeView() const { return CurveEditorTreeView; }
+	
+private:
+	
+	/** The owning Sequencer instance that the Curve Editor is being displayed for. */
+	TWeakPtr<FSequencer> WeakSequencer;
+	/** The curve editor being displayed. */
+	TSharedPtr<FCurveEditor> CurveEditor;
+	
+	/** Used to translate between IOutlinerExtension and FCurveEditorTreeItemID. */
+	TAttribute<const FCurveModelSyncer*> CurveModelSyncerAttr;
+	
+	/** The filter row content displayed above the tree. */
+	TSharedPtr<SCurveEditorSearchAndFilterRow> FilterRow;
+	/** Curve editor tree widget */
+	TSharedPtr<SCurveEditorTree> CurveEditorTreeView;
+	/** The current playback time display. */
+	TSharedPtr<STemporarilyFocusedSpinBox<double>> PlayTimeDisplay;
+
+	TSharedRef<SWidget> MakeFilterArea(const FArguments& InArgs);
+	TSharedRef<SWidget> MakeTreeContent(const FArguments& InArgs);
+	TSharedRef<SWidget> MakeFooter(const FArguments& InArgs);
+	
+	/** Scrolls the given item into view. */
+	void ScrollIntoView(const TWeakViewModelPtr<IOutlinerExtension>& WeakViewModel);
+	/** @return All items visible in the tree view. */
+	TSet<TWeakViewModelPtr<IOutlinerExtension>> GetAllItems() const;
+	/** @return All items selected in the tree view. */
+	TSet<TWeakViewModelPtr<IOutlinerExtension>> GetSelectedItems() const;
+
+	EVisibility GetTransportControlsVisibility() const;
+};
+} // namespace UE::Sequencer
diff --git a/Engine/Source/Editor/Sequencer/Private/CurveEditor/Widgets/Tree/SCurveTreeContent.cpp b/Engine/Source/Editor/Sequencer/Private/CurveEditor/Widgets/Tree/SCurveTreeContent.cpp
deleted file mode 100644
diff --git a/Engine/Source/Editor/Sequencer/Private/CurveEditor/Widgets/Tree/SCurveTreeContent.h b/Engine/Source/Editor/Sequencer/Private/CurveEditor/Widgets/Tree/SCurveTreeContent.h
deleted file mode 100644
diff --git a/Engine/Source/Editor/Sequencer/Private/CurveEditor/Widgets/Tree/SOldCurveTreeContent.cpp b/Engine/Source/Editor/Sequencer/Private/CurveEditor/Widgets/Tree/SOldCurveTreeContent.cpp
deleted file mode 100644
diff --git a/Engine/Source/Editor/Sequencer/Private/CurveEditor/Widgets/Tree/SOldCurveTreeContent.h b/Engine/Source/Editor/Sequencer/Private/CurveEditor/Widgets/Tree/SOldCurveTreeContent.h
deleted file mode 100644
