[Sequencer] TransformSelection function to support empty key selections
- Added static TransformKeySelection() to allow processing on specific keys instead of only the Sequencer key selection
- Added static TransformSelectedOrRelativeKeys that transforms either the selected keys or a relative set of keys based on a time delta and scaling factor. If no selected keys are provided, keys relative to the specified time range will be collected and transformed.

#Synced-CL 52687518

diff --git a/Engine/Source/Editor/Sequencer/Private/Misc/KeyHelperUtils.cpp b/Engine/Source/Editor/Sequencer/Private/Misc/KeyHelperUtils.cpp
new file mode 100644
--- /dev/null
+++ b/Engine/Source/Editor/Sequencer/Private/Misc/KeyHelperUtils.cpp
@@ -0,0 +1,257 @@
+﻿// Copyright Epic Games, Inc. All Rights Reserved.
+
+#include "KeyHelperUtils.h"
+#include "IKeyArea.h"
+#include "ISequencer.h"
+#include "ScopedTransaction.h"
+#include "Sequencer.h"
+#include "SequencerSelectedKey.h"
+
+#define LOCTEXT_NAMESPACE "KeyHelperUtils"
+
+namespace UE::Sequencer::KeyHelperUtils
+{
+
+void CollectKeysRelativeToTime(const TSet<TWeakViewModelPtr<FChannelModel>>& InWeakChannelModels
+	, const TRange<FFrameNumber>& InRange, TSet<FSequencerSelectedKey>& OutSelectedKeys)
+{
+	TArray<FKeyHandle> KeyHandles;
+
+	for (const TWeakViewModelPtr<FChannelModel>& WeakChannelModel : InWeakChannelModels)
+	{
+		const TSharedPtr<FChannelModel> ChannelModel = WeakChannelModel.Pin();
+		if (!ChannelModel)
+		{
+			continue;
+		}
+
+		const TSharedPtr<IKeyArea> KeyArea = ChannelModel->GetKeyArea();
+		if (!KeyArea)
+		{
+			continue;
+		}
+
+		KeyHandles.Reset();
+		KeyArea->GetKeyHandles(KeyHandles, InRange);
+		
+		if (KeyHandles.IsEmpty())
+		{
+			continue;
+		}
+
+		UMovieSceneSection* const Section = KeyArea->GetOwningSection();
+		if (!Section)
+		{
+			continue;
+		}
+
+		for (const FKeyHandle KeyHandle : KeyHandles)
+		{
+			OutSelectedKeys.Add(FSequencerSelectedKey(*Section, ChannelModel, KeyHandle));
+		}
+	}
+}
+
+void TransformKeySelection(ISequencer& InSequencer
+	, const TSet<FSequencerSelectedKey>& InKeys
+	, const FFrameTime& InDeltaTime
+	, const float InScale)
+{
+	using namespace UE::Sequencer;
+
+	if (InKeys.IsEmpty())
+	{
+		return;
+	}
+
+	bool bAnythingChanged = false;
+
+	const FFrameTime OriginTime = InSequencer.GetLocalTime().Time;
+	const bool bApplyScale = (InScale != 0.f);
+
+	TArray<FSequencerSelectedKey> SelectedKeys;
+	SelectedKeys.Reserve(InKeys.Num());
+	SelectedKeys.Append(InKeys.Array());
+	if (SelectedKeys.IsEmpty())
+	{
+		return;
+	}
+
+	FSelectedKeysByChannel KeysByChannel(MakeArrayView(SelectedKeys));
+
+	if (KeysByChannel.SelectedChannels.IsEmpty())
+	{
+		return;
+	}
+
+	TMap<UMovieSceneSection*, TRange<FFrameNumber>> SectionToNewBounds;
+	TArray<TPair<FSelectedChannelInfo*, TArray<FFrameNumber>>> DeferredNewTimesByChannel;
+
+	SectionToNewBounds.Reserve(KeysByChannel.SelectedChannels.Num());
+	if (bApplyScale)
+	{
+		DeferredNewTimesByChannel.Reserve(KeysByChannel.SelectedChannels.Num());
+	}
+
+	auto UpdateSectionBounds = [&SectionToNewBounds]
+		(UMovieSceneSection* InSection, FFrameTime InLowestFrameTime, FFrameTime InHighestFrameTime)
+	{
+		TRange<FFrameNumber>* NewSectionBounds = SectionToNewBounds.Find(InSection);
+		if (!NewSectionBounds)
+		{
+			InSection->Modify();
+			NewSectionBounds = &SectionToNewBounds.Add(InSection, InSection->GetRange());
+		}
+		*NewSectionBounds = TRange<FFrameNumber>::Hull(
+			*NewSectionBounds,
+			TRange<FFrameNumber>(InLowestFrameTime.GetFrame(), InHighestFrameTime.GetFrame() + 1)
+		);
+	};
+
+	auto TransformTimesAndComputeBounds = [OriginTime, InDeltaTime, InScale, bApplyScale]
+		(TArray<FFrameNumber>& InOutTimes) -> TPair<FFrameTime, FFrameTime>
+	{
+		check(!InOutTimes.IsEmpty());
+
+		FFrameTime LowestFrameTime = InOutTimes[0];
+		FFrameTime HighestFrameTime = InOutTimes[0];
+
+		for (FFrameNumber& Time : InOutTimes)
+		{
+			const FFrameTime KeyTime = Time;
+
+			Time = bApplyScale
+				? (OriginTime + InDeltaTime + (KeyTime - OriginTime) * InScale).FloorToFrame()
+				: (KeyTime + InDeltaTime).FloorToFrame();
+
+			if (Time < LowestFrameTime)
+			{
+				LowestFrameTime = Time;
+			}
+			if (Time > HighestFrameTime)
+			{
+				HighestFrameTime = Time;
+			}
+		}
+
+		return TPair<FFrameTime, FFrameTime>(LowestFrameTime, HighestFrameTime);
+	};
+
+	const FScopedTransaction TransformSelectionTransaction(LOCTEXT("TransformSelection_Transaction", "Transform Selection"));
+
+	for (FSelectedChannelInfo& ChannelInfo : KeysByChannel.SelectedChannels)
+	{
+		if (!ChannelInfo.OwningSection || ChannelInfo.KeyHandles.IsEmpty())
+		{
+			continue;
+		}
+
+		UMovieSceneSection* const Section = ChannelInfo.OwningSection;
+		if (Section->IsLocked() || Section->IsReadOnly())
+		{
+			continue;
+		}
+
+		FMovieSceneChannel* const Channel = ChannelInfo.Channel.Get();
+		if (!Channel)
+		{
+			continue;
+		}
+
+		TArray<FFrameNumber> ImmediateNewTimes;
+		TArray<FFrameNumber>* NewTimes = &ImmediateNewTimes;
+		if (bApplyScale)
+		{
+			DeferredNewTimesByChannel.Add(TPair<FSelectedChannelInfo*, TArray<FFrameNumber>>(&ChannelInfo, TArray<FFrameNumber>()));
+			NewTimes = &DeferredNewTimesByChannel.Last().Value;
+		}
+
+		NewTimes->SetNum(ChannelInfo.KeyHandles.Num());
+		Channel->GetKeyTimes(ChannelInfo.KeyHandles, *NewTimes);
+
+		const TPair<FFrameTime, FFrameTime> Bounds = TransformTimesAndComputeBounds(*NewTimes);
+
+		UpdateSectionBounds(Section, Bounds.Key, Bounds.Value);
+		bAnythingChanged = true;
+
+		if (!bApplyScale)
+		{
+			Channel->SetKeyTimes(ChannelInfo.KeyHandles, *NewTimes);
+		}
+	}
+
+	if (bApplyScale)
+	{
+		for (TPair<FSelectedChannelInfo*, TArray<FFrameNumber>>& Pair : DeferredNewTimesByChannel)
+		{
+			FSelectedChannelInfo* const ChannelInfo = Pair.Key;
+			if (!ChannelInfo || ChannelInfo->KeyHandles.IsEmpty())
+			{
+				continue;
+			}
+
+			FMovieSceneChannel* const Channel = ChannelInfo->Channel.Get();
+			if (!Channel)
+			{
+				continue;
+			}
+
+			Channel->SetKeyTimes(ChannelInfo->KeyHandles, Pair.Value);
+		}
+	}
+
+	for (TPair<UMovieSceneSection*, TRange<FFrameNumber>>& Pair : SectionToNewBounds)
+	{
+		Pair.Key->SetRange(Pair.Value);
+	}
+
+	for (TPair<UMovieSceneSection*, TRange<FFrameNumber>>& Pair : SectionToNewBounds)
+	{
+		if (UMovieSceneTrack* Track = Pair.Key->GetTypedOuter<UMovieSceneTrack>())
+		{
+			Track->OnSectionMoved(*Pair.Key, FMovieSceneSectionMovedParams(EPropertyChangeType::ValueSet));
+		}
+	}
+
+	if (bAnythingChanged)
+	{
+		InSequencer.NotifyMovieSceneDataChanged(EMovieSceneDataChangeType::TrackValueChanged);
+	}
+}
+
+bool TransformSelectedOrRelativeKeys(ISequencer& InSequencer
+	, const TSet<TWeakViewModelPtr<FChannelModel>>& InWeakChannels
+	, const TSet<FSequencerSelectedKey>& InSelectedKeys
+	, const FFrameTime InDeltaTime
+	, const float InScale)
+{
+	TSet<FSequencerSelectedKey> KeysToTransform = InSelectedKeys;
+
+	if (KeysToTransform.IsEmpty())
+	{
+		const FFrameNumber ScrubFrame = InSequencer.GetLocalTime().Time.FloorToFrame();
+		const bool bMoveLeft = InDeltaTime < FFrameTime(0);
+		const TRange<FFrameNumber> Range = bMoveLeft
+			? TRange<FFrameNumber>(
+				TRangeBound<FFrameNumber>::Open(),
+				TRangeBound<FFrameNumber>::Exclusive(ScrubFrame))
+			: TRange<FFrameNumber>(
+				TRangeBound<FFrameNumber>::Inclusive(ScrubFrame),
+				TRangeBound<FFrameNumber>::Open());
+
+		CollectKeysRelativeToTime(InWeakChannels, Range, KeysToTransform);
+	}
+
+	if (KeysToTransform.IsEmpty())
+	{
+		return false;
+	}
+
+	TransformKeySelection(InSequencer, KeysToTransform, InDeltaTime, InScale);
+
+	return true;
+}
+
+} // namespace UE::Sequencer::KeyHelperUtils
+
+#undef LOCTEXT_NAMESPACE
diff --git a/Engine/Source/Editor/Sequencer/Private/Misc/KeyHelperUtils.h b/Engine/Source/Editor/Sequencer/Private/Misc/KeyHelperUtils.h
new file mode 100644
--- /dev/null
+++ b/Engine/Source/Editor/Sequencer/Private/Misc/KeyHelperUtils.h
@@ -0,0 +1,65 @@
+﻿// Copyright Epic Games, Inc. All Rights Reserved.
+
+#pragma once
+
+#include "Math/Range.h"
+#include "MVVM/ViewModelPtr.h"
+
+class FSequencer;
+class ISequencer;
+struct FFrameNumber;
+struct FFrameTime;
+struct FSequencerSelectedKey;
+
+namespace UE::Sequencer
+{
+	class FChannelModel;
+}
+
+namespace UE::Sequencer::KeyHelperUtils
+{
+
+/**
+ * Collects sequencer keys within the specified time range from the given channel models
+ * and adds them to the output set of selected keys.
+ *
+ * @param InWeakChannelModels The set of weak pointers to channel models from which keys will be collected.
+ * @param InRange The range of time used to filter the keys for collection.
+ * @param OutSelectedKeys The set where the collected keys will be stored.
+ */
+void CollectKeysRelativeToTime(const TSet<TWeakViewModelPtr<FChannelModel>>& InWeakChannelModels
+	, const TRange<FFrameNumber>& InRange, TSet<FSequencerSelectedKey>& OutSelectedKeys);
+
+/**
+ * Transforms the selected keys by applying a time delta and an optional scaling factor.
+ * Updates the affected sections' bounds and notifies the sequencer of changes.
+ *
+ * @param InSequencer The sequencer instance used to transform the keys.
+ * @param InKeys The set of selected keys to transform.
+ * @param InDeltaTime The time delta used for translating the keys.
+ * @param InScale The scaling factor applied to the keys.
+ */
+void TransformKeySelection(ISequencer& InSequencer
+	, const TSet<FSequencerSelectedKey>& InKeys
+	, const FFrameTime& InDeltaTime
+	, const float InScale);
+
+/**
+ * Transforms either the selected keys or a relative set of keys based on a time delta and scaling factor.
+ * If no selected keys are provided, keys relative to the specified time range will be collected and transformed.
+ *
+ * @param InSequencer The sequencer instance used to transform the keys.
+ * @param InWeakChannels A set of weak pointers to channel models that define the scope for collecting keys.
+ * @param InSelectedKeys A set of keys explicitly selected for transformation.
+ * @param InDeltaTime The time delta used for translating the keys.
+ * @param InScale The scaling factor applied to the keys.
+ * 
+ * @return True if keys were successfully transformed, false otherwise.
+ */
+bool TransformSelectedOrRelativeKeys(ISequencer& InSequencer
+	, const TSet<TWeakViewModelPtr<FChannelModel>>& InWeakChannels
+	, const TSet<FSequencerSelectedKey>& InSelectedKeys
+	, const FFrameTime InDeltaTime
+	, const float InScale);
+
+} // namespace UE::Sequencer::KeyHelperUtils
diff --git a/Engine/Source/Editor/Sequencer/Private/Sequencer.cpp b/Engine/Source/Editor/Sequencer/Private/Sequencer.cpp
--- a/Engine/Source/Editor/Sequencer/Private/Sequencer.cpp
+++ b/Engine/Source/Editor/Sequencer/Private/Sequencer.cpp
@@ -200,6 +200,7 @@
 #include "AssetRegistry/IAssetRegistry.h"
 #include "AssetRegistry/AssetRegistryModule.h"
 #include "Widgets/SSequencerWrapper.h"
+#include "Misc/KeyHelperUtils.h"
 
 #define LOCTEXT_NAMESPACE "Sequencer"
 
@@ -2393,197 +2394,33 @@ void FSequencer::TransformSelection(FFrameTime InDeltaTime, float InScale)
 {
 	using namespace UE::Sequencer;
 
-	FScopedTransaction TransformSelectionTransaction(NSLOCTEXT("Sequencer", "TransformSelection_Transaction", "Transform Selection"));
-	bool bAnythingChanged = false;
-
-	TSet<UMovieSceneSection*> SelectedSections = ViewModel->GetSelection()->GetSelectedSections();
-
-	const FFrameTime OriginTime = GetLocalTime().Time;
-
-	FSelectedKeysByChannel KeysByChannel(ViewModel->GetSelection()->KeySelection);
-	TMap<UMovieSceneSection*, TRange<FFrameNumber>> SectionToNewBounds;
-
-	if (InScale != 0.f)
+	if (!ViewModel.IsValid())
 	{
-		TMap<FMovieSceneChannel*, TPair<TArray<FFrameNumber>, TArray<FKeyHandle>>> ChannelsAndKeyTimes;
-		// Dilate the keys
-		for (const FSelectedChannelInfo& ChannelInfo : KeysByChannel.SelectedChannels)
-		{
-			FMovieSceneChannel* Channel = ChannelInfo.Channel.Get();
-			if (Channel)
-			{
-				// Skip any channels whose section is already selected because they'll be handled below (moving the section and the keys together)
-				if (SelectedSections.Contains(ChannelInfo.OwningSection))
-				{
-					continue;
-				}
-
-				// Skip any locked sections
-				if (ChannelInfo.OwningSection->IsLocked())
-				{
-					continue;
-				}
-
-				TPair<TArray<FFrameNumber>, TArray<FKeyHandle>>& KeyTimesScratch = ChannelsAndKeyTimes.FindOrAdd(Channel);
-				const int32 NumKeys = ChannelInfo.KeyHandles.Num();
-				KeyTimesScratch.Key.Reset(NumKeys);
-				KeyTimesScratch.Key.SetNum(NumKeys);
-				// Populate the key times scratch buffer with the times for these handles
-				KeyTimesScratch.Value = ChannelInfo.KeyHandles;
-				Channel->GetKeyTimes(ChannelInfo.KeyHandles, KeyTimesScratch.Key);
-
-				// We have to find the lowest key time and the highest key time. They're added based on selection order so we can't rely on their order in the array.
-				FFrameTime LowestFrameTime = KeyTimesScratch.Key[0];
-				FFrameTime HighestFrameTime = KeyTimesScratch.Key[0];
-
-				// Perform the transformation
-				for (FFrameNumber& Time : KeyTimesScratch.Key)
-				{
-					FFrameTime KeyTime = Time;
-					Time = (OriginTime + InDeltaTime + (KeyTime - OriginTime) * InScale).FloorToFrame();
-
-					if (Time < LowestFrameTime)
-					{
-						LowestFrameTime = Time;
-					}
-
-					if (Time > HighestFrameTime)
-					{
-						HighestFrameTime = Time;
-					}
-				}
-
-				TRange<FFrameNumber>* NewSectionBounds = SectionToNewBounds.Find(ChannelInfo.OwningSection);
-				if (!NewSectionBounds)
-				{
-					// Call Modify on the owning section before we call SetKeyTimes so that our section bounds/key times stay in sync.
-					ChannelInfo.OwningSection->Modify();
-					NewSectionBounds = &SectionToNewBounds.Add(ChannelInfo.OwningSection, ChannelInfo.OwningSection->GetRange());
-				}
-
-
-				// Expand the range by ensuring the new range contains the range our keys are in. We add one because the highest time is exclusive
-				// for sections, but HighestFrameTime is measuring only the key's time.
-				*NewSectionBounds = TRange<FFrameNumber>::Hull(*NewSectionBounds, TRange<FFrameNumber>(LowestFrameTime.GetFrame(), HighestFrameTime.GetFrame() + 1));
-
-				
-				bAnythingChanged = true;
-			}
-		}
-		// set key times AFTER getting where we want all of them where we want to go, this way we avoid double transforms
-		// on things like smart keys for spaces and constraints.
-		for (TPair<FMovieSceneChannel*, TPair<TArray<FFrameNumber>, TArray<FKeyHandle>>>& ChannelAndKeys : ChannelsAndKeyTimes)
-		{
-			// Apply the new, transformed key times
-			// Channel->SetKeyTimes(KeyHandles,KeyTimes) -->
-			ChannelAndKeys.Key->SetKeyTimes(ChannelAndKeys.Value.Value, ChannelAndKeys.Value.Key);
-		}
-		// Dilate the sections
-		for (UMovieSceneSection* Section : SelectedSections)
-		{
-			// Skip any locked sections
-			if (Section->IsLocked())
-			{
-				continue;
-			}
-
-			TRangeBound<FFrameNumber> LowerBound = Section->GetRange().GetLowerBound();
-			TRangeBound<FFrameNumber> UpperBound = Section->GetRange().GetUpperBound();
-
-			if (Section->HasStartFrame())
-			{
-				FFrameTime StartTime = Section->GetInclusiveStartFrame();
-				FFrameNumber StartFrame = (OriginTime + InDeltaTime + (StartTime - OriginTime) * InScale).FloorToFrame();
-				LowerBound = TRangeBound<FFrameNumber>::Inclusive(StartFrame);
-			}
-
-			if (Section->HasEndFrame())
-			{
-				FFrameTime EndTime = Section->GetExclusiveEndFrame();
-				FFrameNumber EndFrame = (OriginTime + InDeltaTime + (EndTime - OriginTime) * InScale).FloorToFrame();
-				UpperBound = TRangeBound<FFrameNumber>::Exclusive(EndFrame);
-			}
-
-			TRange<FFrameNumber>* NewSectionBounds = SectionToNewBounds.Find(Section);
-			if (!NewSectionBounds)
-			{
-				// Call Modify on the owning section before we call SetKeyTimes so that our section bounds/key times stay in sync.
-				Section->Modify();
-				NewSectionBounds = &SectionToNewBounds.Add( Section, TRange<FFrameNumber>(LowerBound, UpperBound) );
-			}
-
-			// If keys have already modified the section, we're applying the same modification to the section so we can
-			// overwrite the (possibly) existing bound, so it's okay to just overwrite the range without a TRange::Hull.
-			*NewSectionBounds = TRange<FFrameNumber>(LowerBound, UpperBound);
-			bAnythingChanged = true;
-
-			// Modify all of the keys of this section
-			for (const FMovieSceneChannelEntry& Entry : Section->GetChannelProxy().GetAllEntries())
-			{
-				for (FMovieSceneChannel* Channel : Entry.GetChannels())
-				{
-					TArray<FFrameNumber> KeyTimes;
-					TArray<FKeyHandle> KeyHandles;
-					TArray<FFrameNumber> NewKeyTimes;
-					Channel->GetKeys(TRange<FFrameNumber>::All(), &KeyTimes, &KeyHandles);
-
-					for (FFrameNumber KeyTime : KeyTimes)
-					{
-						FFrameNumber NewKeyTime = (OriginTime + InDeltaTime + (KeyTime - OriginTime) * InScale).FloorToFrame();
-						NewKeyTimes.Add(NewKeyTime);
-					}
-
-					Channel->SetKeyTimes(KeyHandles, NewKeyTimes);
-				}
-			}
-
-			Section->BroadcastChanged();
-		}
-
-		// Marked frames
-		const FMarkedFrameSelection& SelectedMarkedFrames = ViewModel->GetSelection()->MarkedFrames;
-		if (SelectedMarkedFrames.Num() > 0)
-		{
-			bAnythingChanged = true;
-
-			UMovieScene* FocusedMovieScene = GetFocusedMovieSceneSequence()->GetMovieScene();
-			FocusedMovieScene->Modify();
-			
-			for (const int32 MarkIndex : SelectedMarkedFrames)
-			{
-				FFrameNumber FrameNumber = FocusedMovieScene->GetMarkedFrames()[MarkIndex].FrameNumber;
-				FrameNumber = (OriginTime + InDeltaTime + (FrameNumber - OriginTime) * InScale).FloorToFrame();
-				FocusedMovieScene->SetMarkedFrame(MarkIndex, FrameNumber);
-			}
-
-			if (!bDraggingMarks)
-			{
-				FocusedMovieScene->SortMarkedFrames();
-			}
-		}
+		return;
 	}
-	
-	// Remove any null sections so we don't need a null check inside the loop.
-	SectionToNewBounds.Remove(nullptr);
-	for (TTuple<UMovieSceneSection*, TRange<FFrameNumber>>& Pair : SectionToNewBounds)
+
+	const TViewModelPtr<FSequenceModel> RootSequenceModel = ViewModel->GetRootSequenceModel();
+	if (!RootSequenceModel.IsValid())
 	{
-		// Set the range of each section that has been modified to their new bounds.
-		Pair.Key->SetRange(Pair.Value);
+		return;
 	}
 
-	// Notify tracks that sections moved so they can update dependent state
-	for (TTuple<UMovieSceneSection*, TRange<FFrameNumber>>& Pair : SectionToNewBounds)
+	TSet<FSequencerSelectedKey> SelectedKeys;
+	FSequencerSelectedKey::AppendKeySelection(SelectedKeys, GetSelection().KeySelection);
+
+	TSet<TSharedPtr<FChannelModel>> AllChannels;
+	SequencerHelpers::GetAllChannels(RootSequenceModel, AllChannels);
+
+	TSet<TWeakViewModelPtr<FChannelModel>> WeakChannels;
+	for (const TSharedPtr<FChannelModel>& Channel : AllChannels)
 	{
-		if (UMovieSceneTrack* Track = Pair.Key->GetTypedOuter<UMovieSceneTrack>())
+		if (Channel.IsValid())
 		{
-			Track->OnSectionMoved(*Pair.Key, FMovieSceneSectionMovedParams(EPropertyChangeType::ValueSet));
+			WeakChannels.Add(TWeakViewModelPtr<FChannelModel>(Channel));
 		}
 	}
 
-	if (bAnythingChanged)
-	{
-		NotifyMovieSceneDataChanged( EMovieSceneDataChangeType::TrackValueChanged );
-	}
+	KeyHelperUtils::TransformSelectedOrRelativeKeys(*this, WeakChannels, SelectedKeys, InDeltaTime, InScale);
 }
 
 void FSequencer::TranslateSelection(bool bTranslateLeft)
diff --git a/Engine/Source/Editor/Sequencer/Private/SimpleView/Menus/SequencerToolbarExtension.cpp b/Engine/Source/Editor/Sequencer/Private/SimpleView/Menus/SequencerToolbarExtension.cpp
--- a/Engine/Source/Editor/Sequencer/Private/SimpleView/Menus/SequencerToolbarExtension.cpp
+++ b/Engine/Source/Editor/Sequencer/Private/SimpleView/Menus/SequencerToolbarExtension.cpp
@@ -15,7 +15,6 @@
 #include "ToolMenus.h"
 #include "ToolableTimeline/Menus/ToolableTimelineMenuContext.h"
 #include "ToolableTimeline/ToolableTimelineCommands.h"
-#include "ToolableTimeline/Tools/ToolableTimelineBaseTool.h"
 
 #define LOCTEXT_NAMESPACE "SequencerToolbarExtension"
 
@@ -92,11 +91,6 @@ void FSequencerToolbarExtension::CreateKeyTranslateSection(FToolMenuSection& InS
 	const TSharedRef<SValueStepper> ValueWidget = SNew(SValueStepper)
 		.NumericTypeInterface(Sequencer->GetNumericTypeInterface())
 		.MinFractionalDigits(0)
-		.IsEnabled_Lambda([WeakTimeline]()
-			{
-				const TSharedPtr<FSimpleViewTimeline> Timeline = WeakTimeline.Pin();
-				return Timeline.IsValid() && Utils::IsSimpleViewToolMenuVisible(WeakTimeline);
-			})
 		.IsVisible_Lambda([WeakTimeline]()
 			{
 				const TSharedPtr<FSimpleViewTimeline> Timeline = WeakTimeline.Pin();
@@ -220,24 +214,6 @@ TSharedPtr<FSequencer> FSequencerToolbarExtension::GetSequencerFromTimeline(cons
 	return Timeline.IsValid() ? Timeline->GetSequencer() : nullptr;
 }
 
-void FSequencerToolbarExtension::TransformSelection(const TWeakPtr<FSimpleViewTimeline> InWeakTimeline
-	, const FFrameTime& InDeltaTime, const float InScale)
-{
-	const TSharedPtr<FSimpleViewTimeline> Timeline = InWeakTimeline.Pin();
-	if (!Timeline.IsValid())
-	{
-		return;
-	}
-
-	const TSharedPtr<FSequencer> Sequencer = GetSequencerFromTimeline(InWeakTimeline);
-	if (!Sequencer.IsValid())
-	{
-		return;
-	}
-
-	Sequencer->TransformSelection(InDeltaTime, InScale);
-}
-
 } // namespace UE::Sequencer::SimpleView
 
 #undef LOCTEXT_NAMESPACE
diff --git a/Engine/Source/Editor/Sequencer/Private/SimpleView/Menus/SequencerToolbarExtension.h b/Engine/Source/Editor/Sequencer/Private/SimpleView/Menus/SequencerToolbarExtension.h
--- a/Engine/Source/Editor/Sequencer/Private/SimpleView/Menus/SequencerToolbarExtension.h
+++ b/Engine/Source/Editor/Sequencer/Private/SimpleView/Menus/SequencerToolbarExtension.h
@@ -7,7 +7,6 @@
 class FSequencer;
 class FUICommandInfo;
 class UToolableTimelineMenuContext;
-struct FFrameTime;
 struct FToolMenuSection;
 
 namespace UE::Sequencer::SimpleView
@@ -38,9 +37,6 @@ protected:
 		, const TArray<TSharedPtr<FUICommandInfo>> InActions);
 
 	static TSharedPtr<FSequencer> GetSequencerFromTimeline(const TWeakPtr<FSimpleViewTimeline>& InWeakTimeline);
-
-	static void TransformSelection(const TWeakPtr<FSimpleViewTimeline> InWeakTimeline
-		, const FFrameTime& InDeltaTime, const float InScale);
 };
 
 } // namespace UE::Sequencer::SimpleView
diff --git a/Engine/Source/Editor/Sequencer/Private/SimpleView/SimpleViewTimeline.cpp b/Engine/Source/Editor/Sequencer/Private/SimpleView/SimpleViewTimeline.cpp
--- a/Engine/Source/Editor/Sequencer/Private/SimpleView/SimpleViewTimeline.cpp
+++ b/Engine/Source/Editor/Sequencer/Private/SimpleView/SimpleViewTimeline.cpp
@@ -1,4 +1,4 @@
-﻿// Copyright Epic Games, Inc. All Rights Reserved.
+// Copyright Epic Games, Inc. All Rights Reserved.
 
 #include "SimpleViewTimeline.h"
 #include "Framework/Docking/SDockingTabStack.h"
@@ -6,6 +6,7 @@
 #include "Sequencer.h"
 #include "SequencerSettings.h"
 #include "SequencerSimpleViewSettings.h"
+#include "Misc/KeyHelperUtils.h"
 #include "MVVM/Selection/Selection.h"
 #include "SimpleView/Menus/ToolableTimelineMenuExtension.h"
 #include "SimpleView/SimpleViewCommands.h"
@@ -113,11 +114,9 @@ void FSimpleViewTimeline::BindSimpleViewCommands()
 		, FExecuteAction::CreateSP(this, &FSimpleViewTimeline::ToggleToolbarVisible));
 
 	CommandList->MapAction(SimpleViewEditorCommands.TranslateKeyLeft
-		, FExecuteAction::CreateSP(this, &FSimpleViewTimeline::TranslateKeyLeft)
-		, FCanExecuteAction::CreateSP(this, &FSimpleViewTimeline::HasSequencerKeySelection));
+		, FExecuteAction::CreateSP(this, &FSimpleViewTimeline::TranslateKeyLeft));
 	CommandList->MapAction(SimpleViewEditorCommands.TranslateKeyRight
-		, FExecuteAction::CreateSP(this, &FSimpleViewTimeline::TranslateKeyRight)
-		, FCanExecuteAction::CreateSP(this, &FSimpleViewTimeline::HasSequencerKeySelection));
+		, FExecuteAction::CreateSP(this, &FSimpleViewTimeline::TranslateKeyRight));
 
 	CommandList->MapAction(SimpleViewEditorCommands.ScaleKeyMultiply
 		, FExecuteAction::CreateSP(this, &FSimpleViewTimeline::ScaleKeyMultiply)
@@ -314,10 +313,13 @@ void FSimpleViewTimeline::TranslateKeyLeft()
 	}
 
 	const FFrameNumber DeltaTime = FFrameTime::FromDecimal(KeyTranslateDelta).GetFrame();
-	if (DeltaTime != 0)
+	if (DeltaTime == 0)
 	{
-		Sequencer->TransformSelection(-DeltaTime, 1.f);
+		return;
 	}
+
+	KeyHelperUtils::TransformSelectedOrRelativeKeys(*Sequencer
+		, GetChannelModels(), KeySelection.GetSelectedKeys(), -DeltaTime, 1.f);
 }
 
 void FSimpleViewTimeline::TranslateKeyRight()
@@ -329,10 +331,13 @@ void FSimpleViewTimeline::TranslateKeyRight()
 	}
 
 	const FFrameNumber DeltaTime = FFrameTime::FromDecimal(KeyTranslateDelta).GetFrame();
-	if (DeltaTime != 0)
+	if (DeltaTime == 0)
 	{
-		Sequencer->TransformSelection(DeltaTime, 1.f);
+		return;
 	}
+
+	KeyHelperUtils::TransformSelectedOrRelativeKeys(*Sequencer
+		, GetChannelModels(), KeySelection.GetSelectedKeys(), DeltaTime, 1.f);
 }
 
 void FSimpleViewTimeline::ScaleKeyDivide()
diff --git a/Engine/Source/Editor/Sequencer/Private/ToolableTimeline/ToolableTimeline.h b/Engine/Source/Editor/Sequencer/Private/ToolableTimeline/ToolableTimeline.h
--- a/Engine/Source/Editor/Sequencer/Private/ToolableTimeline/ToolableTimeline.h
+++ b/Engine/Source/Editor/Sequencer/Private/ToolableTimeline/ToolableTimeline.h
@@ -315,6 +315,9 @@ protected:
 	/** Event that broadcasts when channel models have changed */
 	FOnChannelModelsChanged ChannelModelsChangedEvent;
 
+	/** Selected keys in the timeline */
+	FToolableTimelineKeySelection KeySelection;
+
 private:
 	friend class FToolableTimeSliderController;
 
@@ -328,9 +331,6 @@ private:
 
 	/** Main widget for this timeline */
 	TSharedPtr<SToolableTimeline> TimelineWidget;
-
-	/**  */
-	FToolableTimelineKeySelection KeySelection;
 };
 
 } // namespace UE::Sequencer::ToolableTimeline
