[SequencerSimpleView] Use newly added SDockTab AutosizePolicy (CL 52686181) to set Sequencer tab to autosize its height.
#Synced-CL 52689698

diff --git a/Engine/Source/Editor/Sequencer/Private/SimpleView/SimpleViewTabAutosizeHelper.cpp b/Engine/Source/Editor/Sequencer/Private/SimpleView/SimpleViewTabAutosizeHelper.cpp
--- a/Engine/Source/Editor/Sequencer/Private/SimpleView/SimpleViewTabAutosizeHelper.cpp
+++ b/Engine/Source/Editor/Sequencer/Private/SimpleView/SimpleViewTabAutosizeHelper.cpp
@@ -142,7 +142,7 @@ EActiveTimerReturnType FSimpleViewTabAutosizeHelper::HandleAutosizeTimer(const d
 	}
 	else
 	{
-		ExistingTab->SetShouldAutosize(true);
+		ExistingTab->SetAutosizePolicy(EDockTabAutosizePolicy::VerticalOnly);
 
 		ExistingTab->Invalidate(EInvalidateWidgetReason::Layout);
 		ExistingTab->SlatePrepass(FSlateApplicationBase::Get().GetApplicationScale());
@@ -198,7 +198,7 @@ void FSimpleViewTabAutosizeHelper::AutoSize(const TSharedRef<SDockTab>& InDockTa
 		UnregisterAutoSizeTimer(InDockTab);
 		UnregisterResetAutosizeTimer(InDockTab);
 
-		InDockTab->SetShouldAutosize(false);
+		InDockTab->SetAutosizePolicy(EDockTabAutosizePolicy::None);
 		InDockTab->Invalidate(EInvalidateWidgetReason::Layout);
 	}
 }
@@ -213,7 +213,9 @@ void FSimpleViewTabAutosizeHelper::RegisterResetAutosizeTimer(const TSharedRef<S
 			{
 				if (const TSharedPtr<SDockTab> Tab = WeakTab.Pin())
 				{
-					Tab->SetShouldAutosize(bInAutosize);
+					const EDockTabAutosizePolicy AutosizePolicy = bInAutosize
+						? EDockTabAutosizePolicy::VerticalOnly : EDockTabAutosizePolicy::None;
+					Tab->SetAutosizePolicy(AutosizePolicy);
 					Tab->Invalidate(EInvalidateWidgetReason::Layout);
 
 					UnregisterResetAutosizeTimer(Tab.ToSharedRef());
diff --git a/Engine/Source/Editor/Sequencer/Private/SimpleView/SimpleViewTimeline.cpp b/Engine/Source/Editor/Sequencer/Private/SimpleView/SimpleViewTimeline.cpp
--- a/Engine/Source/Editor/Sequencer/Private/SimpleView/SimpleViewTimeline.cpp
+++ b/Engine/Source/Editor/Sequencer/Private/SimpleView/SimpleViewTimeline.cpp
@@ -193,9 +193,7 @@ void FSimpleViewTimeline::EnableSimpleView(const bool bInEnable)
 
 			FSimpleViewTabAutosizeHelper::SetParentDockTabStackTabWellHiddenIfOnlyTab(ExistingTabRef, true);
 
-			// This needs to be changed to the new API after that CL goes in
-			ExistingTabRef->SetShouldAutosize(true);
-			//ExistingTabRef->SetAutosizePolicy(EDockTabAutosizePolicy::VerticalOnly);
+			ExistingTabRef->SetAutosizePolicy(EDockTabAutosizePolicy::VerticalOnly);
 
 			/*const bool bFloating = TabAutosizer.IsTabFloating(ExistingTabRef);
 			if (bFloating)
@@ -221,9 +219,7 @@ void FSimpleViewTimeline::EnableSimpleView(const bool bInEnable)
 
 			FSimpleViewTabAutosizeHelper::SetParentDockTabStackTabWellHiddenIfOnlyTab(ExistingTab.ToSharedRef(), false);
 
-			// This needs to be changed to the new API after that CL goes in
-			ExistingTab->SetShouldAutosize(false);
-			//ExistingTabRef->SetAutosizePolicy(EDockTabAutosizePolicy::None);
+			ExistingTabRef->SetAutosizePolicy(EDockTabAutosizePolicy::None);
 
 			const TSharedRef<SWidget> WidgetToFucus = Sequencer->GetSequencerWidget();
 			FSlateApplication::Get().SetKeyboardFocus(WidgetToFucus);
