Sequencer: Fix failing unit test due to not finding config entry.

It seems right to just remove the error log: if there is no config entry, no filters are active, yet. Hence the widget will display empty. That seems correctl and does not need any error.

#Synced-CL 52617198

diff --git a/Engine/Source/Editor/Sequencer/Private/Filters/Widgets/SSequencerFilterBar.cpp b/Engine/Source/Editor/Sequencer/Private/Filters/Widgets/SSequencerFilterBar.cpp
--- a/Engine/Source/Editor/Sequencer/Private/Filters/Widgets/SSequencerFilterBar.cpp
+++ b/Engine/Source/Editor/Sequencer/Private/Filters/Widgets/SSequencerFilterBar.cpp
@@ -376,19 +376,14 @@ void SSequencerFilterBar::CreateFilterWidgetsFromConfig()
 	USequencerSettings* const SequencerSettings = FilterBar->GetSequencer().GetSequencerSettings();
 	check(IsValid(SequencerSettings));
 
+	constexpr bool bSaveConfig = false;
 	const FName InstanceIdentifier = FilterBar->GetIdentifier();
-	FSequencerFilterBarConfig* const Config = SequencerSettings->FindTrackFilterBar(InstanceIdentifier);
-	if (!Config)
-	{
-		UE_LOGF(LogSequencer, Error, "SSequencerFilterBar requires that you specify a FilterBarIdentifier to load settings");
-		return;
-	}
+	const FSequencerFilterBarConfig& Config = SequencerSettings->FindOrAddTrackFilterBar(InstanceIdentifier, bSaveConfig);
 
 	RemoveAllFilterWidgets();
 
 	const TSet<TSharedRef<FFilterCategory>> DisplayableCategories = FilterBar->GetConfigCategories();
-
-	auto LoadFilterFromConfig = [this, &Config, &DisplayableCategories](const TSharedRef<FSequencerTrackFilter>& InFilter)
+	const auto LoadFilterFromConfig = [this, &Config, &DisplayableCategories](const TSharedRef<FSequencerTrackFilter>& InFilter)
 	{
 		const TSharedPtr<FFilterCategory> FilterCategory = InFilter->GetCategory();
 		if (FilterCategory.IsValid() && !DisplayableCategories.Contains(FilterCategory.ToSharedRef()))
@@ -397,7 +392,7 @@ void SSequencerFilterBar::CreateFilterWidgetsFromConfig()
 		}
 
 		const FString FilterName = InFilter->GetDisplayName().ToString();
-		if (!Config->IsFilterEnabled(FilterName))
+		if (!Config.IsFilterEnabled(FilterName))
 		{
 			return;
 		}
