Sequencer: IOutlinerExtensions that are created after a filter runs are now no longer considered filtered out. To achieve this, FSequencerFilterData now tracks which items have been processed already.

This way, the creation is sticky. Use case: Enable "Selected" filter, click the "chest" control rig control, then add a timewarp to the owning dragon. You want the created item to be displayed in the outliner - the item should not immediately disappear.

This fixes adding a TimeWarp to Sequencer is not updating curve editor properly.
The problem was that after adding a time warp channel, FCurveModelSyncer::SyncSequencerToCurveEditor invoked ILinkedFilterViewModel::IsFilteredOut, which returned false because FSequencerFilterData::FilterInNodes did not contain the node... because it was new.

#Synced-CL 52617198

diff --git a/Engine/Source/Editor/Sequencer/Private/Filters/SequencerFilterData.cpp b/Engine/Source/Editor/Sequencer/Private/Filters/SequencerFilterData.cpp
--- a/Engine/Source/Editor/Sequencer/Private/Filters/SequencerFilterData.cpp
+++ b/Engine/Source/Editor/Sequencer/Private/Filters/SequencerFilterData.cpp
@@ -28,8 +28,9 @@ bool FSequencerFilterData::operator!=(const FSequencerFilterData& InRhs) const
 
 void FSequencerFilterData::Reset()
 {
+	ProcessedNodes.Reset();
 	FilterInNodes.Reset();
-
+	
 	TotalNodeCount = 0;
 }
 
@@ -70,6 +71,8 @@ void FSequencerFilterData::AddTotalNodeCount(int32 InNodeCount)
 
 void FSequencerFilterData::FilterInNode(TWeakViewModelPtr<IOutlinerExtension> InNodeWeak)
 {
+	ProcessedNodes.Add(InNodeWeak);
+	
 	FilterInNodes.Add(InNodeWeak);
 
 	if (const TViewModelPtr<IOutlinerExtension> Node = InNodeWeak.Pin())
@@ -80,6 +83,8 @@ void FSequencerFilterData::FilterInNode(TWeakViewModelPtr<IOutlinerExtension> In
 
 void FSequencerFilterData::FilterOutNode(TWeakViewModelPtr<IOutlinerExtension> InNodeWeak)
 {
+	ProcessedNodes.Add(InNodeWeak);
+	
 	const FSetElementId ElementId = FilterInNodes.FindId(InNodeWeak);
 	if (ElementId.IsValidId())
 	{
@@ -127,7 +132,8 @@ bool FSequencerFilterData::ContainsFilterInNodes(const FSequencerFilterData& InO
 
 bool FSequencerFilterData::IsFilteredOut(const TViewModelPtr<IOutlinerExtension>& InNode) const
 {
-	return !FilterInNodes.Contains(InNode);
+	// If an item has been added since the last filtered run, then we'll not consider it filtered out.
+	return ProcessedNodes.Contains(InNode) && !FilterInNodes.Contains(InNode);
 }
 
 TWeakViewModelPtr<ITrackExtension> FSequencerFilterData::ResolveTrack(FSequencerTrackFilterType InNode)
diff --git a/Engine/Source/Editor/Sequencer/Public/Filters/SequencerFilterData.h b/Engine/Source/Editor/Sequencer/Public/Filters/SequencerFilterData.h
--- a/Engine/Source/Editor/Sequencer/Public/Filters/SequencerFilterData.h
+++ b/Engine/Source/Editor/Sequencer/Public/Filters/SequencerFilterData.h
@@ -78,7 +78,7 @@ struct FSequencerFilterData
 	UE_API bool ContainsFilterInNodes(const FSequencerFilterData& InOtherData) const;
 
 	UE_API bool IsFilteredOut(const UE::Sequencer::TViewModelPtr<UE::Sequencer::IOutlinerExtension>& InNode) const;
-
+	
 	UE_API UE::Sequencer::TWeakViewModelPtr<UE::Sequencer::ITrackExtension> ResolveTrack(FSequencerTrackFilterType InNode);
 	UE_API UMovieSceneTrack* ResolveMovieSceneTrackObject(FSequencerTrackFilterType InNode);
 	UE_API UObject* ResolveTrackBoundObject(ISequencer& InSequencer, FSequencerTrackFilterType InNode);
@@ -97,6 +97,11 @@ protected:
 
 	/** Nodes to be displayed in the UI */
 	TSet<UE::Sequencer::TWeakViewModelPtr<UE::Sequencer::IOutlinerExtension>> FilterInNodes;
+	/**
+	 * Nodes that have been processed, i.e. either filtered in or out. 
+	 * This is used to detect when new IOutlinerExtensions are created after filters have run.
+	 */
+	TSet<UE::Sequencer::TWeakViewModelPtr<UE::Sequencer::IOutlinerExtension>> ProcessedNodes;
 	
 	/** Broadcasts when the filter state of IOutlinerExtension is changed. */
 	UE::Sequencer::FChangeOutlinerExtensionFilterState ChangeFilterStateEvent;
