Curve Editor: Toolbar Tweaks

Remove All Edits button which isn't necessary for Curve Editor
Remove Curves text for the pulldown

#Synced-CL 52681601

diff --git a/Engine/Source/Editor/CurveEditor/Private/SCurveEditorPanel.cpp b/Engine/Source/Editor/CurveEditor/Private/SCurveEditorPanel.cpp
--- a/Engine/Source/Editor/CurveEditor/Private/SCurveEditorPanel.cpp
+++ b/Engine/Source/Editor/CurveEditor/Private/SCurveEditorPanel.cpp
@@ -1547,7 +1547,6 @@ void SCurveEditorPanel::BuildToolbar(FToolBarBuilder& InToolBarBuilder, TSharedP
 
 	InToolBarBuilder.BeginSection("Curves");
 	{
-		InToolBarBuilder.SetLabelVisibility(EVisibility::Visible); // Show label...
 		InToolBarBuilder.AddComboButton(
 			FUIAction(),
 			FOnGetContent::CreateSP(this, &SCurveEditorPanel::MakeCurvesMenu, InBaseExtender),
diff --git a/Engine/Source/Editor/Sequencer/Private/CurveEditor/Specializations/SequencerCurveEditorToolbarExtender.h b/Engine/Source/Editor/Sequencer/Private/CurveEditor/Specializations/SequencerCurveEditorToolbarExtender.h
--- a/Engine/Source/Editor/Sequencer/Private/CurveEditor/Specializations/SequencerCurveEditorToolbarExtender.h
+++ b/Engine/Source/Editor/Sequencer/Private/CurveEditor/Specializations/SequencerCurveEditorToolbarExtender.h
@@ -26,7 +26,8 @@ public:
 				
 				ToolbarBuilder.BeginSection("Keying");
 				ToolbarBuilder.PushCommandList(SequenerPin->GetCommandBindings().ToSharedRef());
-				AppendSequencerToolbarEntries(SequenerPin, ToolbarBuilder);
+				EToolbarItemFlags Flags = EToolbarItemFlags::KeyGroup | EToolbarItemFlags::AutoKey;
+				AppendSequencerToolbarEntries(SequenerPin, ToolbarBuilder, Flags);
 				ToolbarBuilder.PopCommandList();
 				ToolbarBuilder.EndSection();
 			}));
