[Sequencer]
- Disable "Popout Transport Controls" if Transport Controls are not visible
- Popout transport controls on toggle of transport controls visibility if the popout option is enabled
#Synced-CL 52691839

diff --git a/Engine/Source/Editor/Sequencer/Private/Misc/TransportControlsHelper.cpp b/Engine/Source/Editor/Sequencer/Private/Misc/TransportControlsHelper.cpp
--- a/Engine/Source/Editor/Sequencer/Private/Misc/TransportControlsHelper.cpp
+++ b/Engine/Source/Editor/Sequencer/Private/Misc/TransportControlsHelper.cpp
@@ -6,8 +6,6 @@
 #include "LevelSequence.h"
 #include "SequencerSettings.h"
 #include "SPopoutTabInlineContent.h"
-#include "ToolMenus.h"
-#include "Widgets/Layout/SScaleBox.h"
 #include "WorkspaceMenuStructure.h"
 #include "WorkspaceMenuStructureModule.h"
 
@@ -37,15 +35,26 @@ TSharedRef<SPopoutTabInlineContent> MakePopoutTransportControls(const TSharedRef
 		.AutosizeWhenFloating(true)
 		.HideDockTabStackTab(true)
 		.ContentHAlign(HAlign_Fill)
-		.InitiallyPoppedOut(SequencerSettings->AreTransportControlsPoppedOut())
-		.OnPopoutStateChanged_Lambda([WeakSequencer](const bool bPoppedOut)
+		.InitiallyPoppedOut(SequencerSettings->ShouldRestoreTransportControlsPoppedOut())
+		.OnPopoutStateChanged_Lambda([WeakSequencer](const bool bInPoppedOut)
 			{
 				const TSharedPtr<ISequencer> Sequencer = WeakSequencer.Pin();
 				USequencerSettings* const SequencerSettings = Sequencer.IsValid()
 					? Sequencer->GetSequencerSettings() : nullptr;
-				if (SequencerSettings)
+				if (!SequencerSettings)
 				{
-					SequencerSettings->SetTransportControlsPoppedOut(bPoppedOut);
+					return;
+				}
+
+				// Live current state
+				SequencerSettings->SetTransportControlsPoppedOut(bInPoppedOut);
+
+				// Only update the restore preference when transport controls are actually being shown.
+				// If they are being hidden, the popout is closed as a side effect and should not erase
+				// the user's preferred popped-out mode.
+				if (SequencerSettings->ShouldShowTransportControls())
+				{
+					SequencerSettings->SetRestoreTransportControlsPoppedOut(bInPoppedOut);
 				}
 			})
 		[
@@ -70,7 +79,9 @@ bool IsPopoutTransportControlsVisible(const TSharedRef<ISequencer>& InSequencer
 
 	// Sequencer transport popped out state overrides curve editor transport controls visibility
 	const USequencerSettings* const SequencerSettings = InSequencer->GetSequencerSettings();
-	if (!SequencerSettings || SequencerSettings->AreTransportControlsPoppedOut())
+	if (!SequencerSettings
+		|| !SequencerSettings->ShouldShowTransportControls()
+		|| SequencerSettings->ShouldRestoreTransportControlsPoppedOut())
 	{
 		return false;
 	}
diff --git a/Engine/Source/Editor/Sequencer/Private/SSequencer.cpp b/Engine/Source/Editor/Sequencer/Private/SSequencer.cpp
--- a/Engine/Source/Editor/Sequencer/Private/SSequencer.cpp
+++ b/Engine/Source/Editor/Sequencer/Private/SSequencer.cpp
@@ -4472,6 +4472,54 @@ FOptionalSize SSequencer::GetOutlinerWidth() const
 	return FMath::Max<float>(TreeViewWidth, MinOutlinerWidth);
 }
 
+bool SSequencer::IsTransportControlsVisible() const
+{
+	const USequencerSettings* const SequencerSettings = GetSequencerSettings();
+	return SequencerSettings && SequencerSettings->ShouldShowTransportControls();
+}
+
+void SSequencer::SetTransportControlsVisible(const bool bInVisible)
+{
+	USequencerSettings* const SequencerSettings = GetSequencerSettings();
+	if (!SequencerSettings || SequencerSettings->ShouldShowTransportControls() == bInVisible)
+	{
+		return;
+	}
+
+	SequencerSettings->SetShowTransportControls(bInVisible);
+
+	if (bInVisible)
+	{
+		const bool bShouldRestorePoppedOut = SequencerSettings->ShouldRestoreTransportControlsPoppedOut();
+
+		// Restore the live state from the hidden persisted preference
+		SequencerSettings->SetTransportControlsPoppedOut(bShouldRestorePoppedOut);
+
+		if (TransportControlsPopoutTab.IsValid())
+		{
+			TransportControlsPopoutTab->Popout(bShouldRestorePoppedOut);
+		}
+	}
+	else
+	{
+		// Hide the visible popout window if present, but do not erase the restore preference
+		if (TransportControlsPopoutTab.IsValid())
+		{
+			TransportControlsPopoutTab->Popout(false);
+		}
+
+		// Live state becomes "not currently popped out" while hidden
+		SequencerSettings->SetTransportControlsPoppedOut(false);
+	}
+
+	RebuildMainContent();
+}
+
+void SSequencer::ToggleTransportControlsVisible()
+{
+	SetTransportControlsVisible(!IsTransportControlsVisible());
+}
+
 bool SSequencer::IsTransportControlsPoppedOut() const
 {
 	return TransportControlsPopoutTab.IsValid() ? TransportControlsPopoutTab->IsPoppedOut() : false;
@@ -4485,17 +4533,41 @@ void SSequencer::SetTransportControlsPoppedOut(const bool bInPoppedOut)
 		return;
 	}
 
+	const bool bWasTransportControlsPoppedOut = IsTransportControlsPoppedOut();
+
+	// Update both the live state and the hidden persisted restore preference
 	SequencerSettings->SetTransportControlsPoppedOut(bInPoppedOut);
+	SequencerSettings->SetRestoreTransportControlsPoppedOut(bInPoppedOut);
 
-	if (TransportControlsPopoutTab.IsValid())
+	if (SequencerSettings->ShouldShowTransportControls() && TransportControlsPopoutTab.IsValid())
 	{
 		TransportControlsPopoutTab->Popout(bInPoppedOut);
 	}
+
+	if (bInPoppedOut)
+	{
+		if (!bWasTransportControlsPoppedOut)
+		{
+			SetSelectionDetailsPoppedOut(true);
+		}
+	}
+	else
+	{
+		if (bWasTransportControlsPoppedOut)
+		{
+			SetSelectionDetailsPoppedOut(false);
+		}
+	}
+
+	RebuildMainContent();
 }
 
 void SSequencer::ToggleTransportControlsPoppedOut()
 {
-	SetTransportControlsPoppedOut(!IsTransportControlsPoppedOut());
+	if (const USequencerSettings* const SequencerSettings = GetSequencerSettings())
+	{
+		SetTransportControlsPoppedOut(!SequencerSettings->ShouldRestoreTransportControlsPoppedOut());
+	}
 }
 
 EVisibility SSequencer::GetPopoutTransportControlsVisibility() const
diff --git a/Engine/Source/Editor/Sequencer/Private/SSequencer.h b/Engine/Source/Editor/Sequencer/Private/SSequencer.h
--- a/Engine/Source/Editor/Sequencer/Private/SSequencer.h
+++ b/Engine/Source/Editor/Sequencer/Private/SSequencer.h
@@ -696,6 +696,10 @@ public:
 	virtual void EnablePendingFocusOnHovering(const bool InEnabled) override;
 	
 	TSharedPtr<FSequencerFilterBar> GetFilterBar() const;
+	
+	bool IsTransportControlsVisible() const;
+	void SetTransportControlsVisible(const bool bInVisible);
+	void ToggleTransportControlsVisible();
 
 	bool IsTransportControlsPoppedOut() const;
 	void SetTransportControlsPoppedOut(const bool bInPoppedOut);
diff --git a/Engine/Source/Editor/Sequencer/Private/Sequencer.cpp b/Engine/Source/Editor/Sequencer/Private/Sequencer.cpp
--- a/Engine/Source/Editor/Sequencer/Private/Sequencer.cpp
+++ b/Engine/Source/Editor/Sequencer/Private/Sequencer.cpp
@@ -11737,15 +11737,15 @@ void FSequencer::BindCommands()
 
 	SequencerCommandBindings->MapAction(
 		Commands.ToggleShowTransportControl,
-		FExecuteAction::CreateLambda( [this]{ Settings->SetShowTransportControls( !Settings->ShouldShowTransportControls() ); } ),
+		FExecuteAction::CreateLambda( [this]{ SequencerWidget->ToggleTransportControlsVisible(); } ),
 		FCanExecuteAction(),
-		FIsActionChecked::CreateLambda( [this]{ return Settings->ShouldShowTransportControls(); } ) );
+		FIsActionChecked::CreateLambda( [this]{ return SequencerWidget->IsTransportControlsVisible(); } ) );
 
 	SequencerCommandBindings->MapAction(
 		Commands.ToggleTransportControlsPoppedOut,
 		FExecuteAction::CreateLambda( [this]{ SequencerWidget->ToggleTransportControlsPoppedOut(); } ),
-		FCanExecuteAction(),
-		FIsActionChecked::CreateLambda( [this]{ return SequencerWidget->IsTransportControlsPoppedOut(); } ) );
+		FCanExecuteAction::CreateLambda( [this]{ return Settings->ShouldShowTransportControls(); } ),
+		FIsActionChecked::CreateLambda( [this]{ return Settings->ShouldRestoreTransportControlsPoppedOut(); } ) );
 
 	SequencerCommandBindings->MapAction(
 		Commands.ToggleShowRangeSlider,
diff --git a/Engine/Source/Editor/Sequencer/Private/Sequencer.h b/Engine/Source/Editor/Sequencer/Private/Sequencer.h
--- a/Engine/Source/Editor/Sequencer/Private/Sequencer.h
+++ b/Engine/Source/Editor/Sequencer/Private/Sequencer.h
@@ -1295,6 +1295,9 @@ public:
 	/** Toggles the state of the Sequencer selection details between docked in the Sequencer tab and being popped out to its own tab window. */
 	void ToggleSelectionDetailsPoppedOut();
 
+	/** @return True if the Sequencer selection details can be popped out to its own tab window. */
+	bool CanSelectionDetailsPopOut() const;
+
 	/** @return True if the Sequencer selection details are currently popped out to its own tab window. */
 	bool IsSelectionDetailsPoppedOut() const;
 
diff --git a/Engine/Source/Editor/Sequencer/Private/SequencerSettings.cpp b/Engine/Source/Editor/Sequencer/Private/SequencerSettings.cpp
--- a/Engine/Source/Editor/Sequencer/Private/SequencerSettings.cpp
+++ b/Engine/Source/Editor/Sequencer/Private/SequencerSettings.cpp
@@ -90,6 +90,7 @@ USequencerSettings::USequencerSettings( const FObjectInitializer& ObjectInitiali
 	bSelectionDetailsPoppedOut = false;
 	bShowTransportControls = true;
 	bTransportControlsPoppedOut = false;
+	bRestoreTransportControlsPoppedOut = false;
 
 	SectionColorTints.Add(FColor(88, 102, 142, 255)); // blue
 	SectionColorTints.Add(FColor(99, 137, 132, 255)); // blue-green
@@ -1439,6 +1440,20 @@ void USequencerSettings::SetTransportControlsPoppedOut(const bool bInPoppedOut)
 	}
 }
 
+bool USequencerSettings::ShouldRestoreTransportControlsPoppedOut() const
+{
+	return bRestoreTransportControlsPoppedOut;
+}
+
+void USequencerSettings::SetRestoreTransportControlsPoppedOut(const bool bInRestorePoppedOut)
+{
+	if (bRestoreTransportControlsPoppedOut != bInRestorePoppedOut)
+	{
+		bRestoreTransportControlsPoppedOut = bInRestorePoppedOut;
+		SaveConfig();
+	}
+}
+
 void USequencerSettings::SetIsSimpleView(const bool bInSimpleView)
 {
 	if (bSimpleView != bInSimpleView)
diff --git a/Engine/Source/Editor/Sequencer/Public/SequencerSettings.h b/Engine/Source/Editor/Sequencer/Public/SequencerSettings.h
--- a/Engine/Source/Editor/Sequencer/Public/SequencerSettings.h
+++ b/Engine/Source/Editor/Sequencer/Public/SequencerSettings.h
@@ -639,6 +639,11 @@ public:
 	/** Sets the popped out state of the transport controls panel in the Sequencer tab */
 	UE_API void SetTransportControlsPoppedOut(const bool bInPoppedOut);
 
+	/** @return True if the persisted preference for restoring transport controls in a popped-out state when shown. */
+	bool ShouldRestoreTransportControlsPoppedOut() const;
+	/** Sets the persisted preference for restoring transport controls in a popped-out state when shown. */
+	void SetRestoreTransportControlsPoppedOut(const bool bInRestorePoppedOut);
+
 	/** @return True if the last the time Sequencer was open, it was in the simplified view and not the full view */
 	bool IsSimpleView() const { return bSimpleView; }
 	/** Switches between full Sequencer and a simplified Sequencer view */
@@ -1037,6 +1042,10 @@ protected:
 	UPROPERTY(config, EditAnywhere, Category = General)
 	bool bTransportControlsPoppedOut;
 
+	/** Hidden persisted preference for whether transport controls should reopen popped out when shown again */
+	UPROPERTY(config)
+	bool bRestoreTransportControlsPoppedOut = false;
+
 	/** Last saved state of the Sequencer view: full or simple */
 	UPROPERTY(config)
 	bool bSimpleView;
