[Sequencer] Change TransformSelectedOrRelativeKeys() to only operate on keys to the right of the scrub head
#Synced-CL 52691839

diff --git a/Engine/Source/Editor/Sequencer/Private/Misc/KeyHelperUtils.cpp b/Engine/Source/Editor/Sequencer/Private/Misc/KeyHelperUtils.cpp
--- a/Engine/Source/Editor/Sequencer/Private/Misc/KeyHelperUtils.cpp
+++ b/Engine/Source/Editor/Sequencer/Private/Misc/KeyHelperUtils.cpp
@@ -230,14 +230,10 @@ bool TransformSelectedOrRelativeKeys(ISequencer& InSequencer
 	if (KeysToTransform.IsEmpty())
 	{
 		const FFrameNumber ScrubFrame = InSequencer.GetLocalTime().Time.FloorToFrame();
-		const bool bMoveLeft = InDeltaTime < FFrameTime(0);
-		const TRange<FFrameNumber> Range = bMoveLeft
-			? TRange<FFrameNumber>(
-				TRangeBound<FFrameNumber>::Open(),
-				TRangeBound<FFrameNumber>::Exclusive(ScrubFrame))
-			: TRange<FFrameNumber>(
-				TRangeBound<FFrameNumber>::Inclusive(ScrubFrame),
-				TRangeBound<FFrameNumber>::Open());
+		const TRange<FFrameNumber> Range(
+			TRangeBound<FFrameNumber>::Exclusive(ScrubFrame),
+			TRangeBound<FFrameNumber>::Open()
+		);
 
 		CollectKeysRelativeToTime(InWeakChannels, Range, KeysToTransform);
 	}
