Sequencer: After opening Sequencer, instanced filter shows nothing until you refresh filters

#Synced-CL 52617198

diff --git a/Engine/Source/Editor/Sequencer/Private/CurveEditor/Integration/Selection/SelectionSyncRootLogic.cpp b/Engine/Source/Editor/Sequencer/Private/CurveEditor/Integration/Selection/SelectionSyncRootLogic.cpp
--- a/Engine/Source/Editor/Sequencer/Private/CurveEditor/Integration/Selection/SelectionSyncRootLogic.cpp
+++ b/Engine/Source/Editor/Sequencer/Private/CurveEditor/Integration/Selection/SelectionSyncRootLogic.cpp
@@ -37,8 +37,6 @@ FSelectionSyncRootLogic::FSelectionSyncRootLogic(
 	CurveEditorToSequencerSyncer.OnEndSync().AddRaw(this, &FSelectionSyncRootLogic::OnEndCurveEditorToSequencerSync);
 	
 	InSequencer->GetViewModel()->GetSelection()->Outliner.OnChanged.AddRaw(this, &FSelectionSyncRootLogic::OnOutlinerSelectionChanged);
-	
-	SequencerToCurveEditorSyncer.UpdateCurveEditor();
 }
 
 FSelectionSyncRootLogic::~FSelectionSyncRootLogic()
diff --git a/Engine/Source/Editor/Sequencer/Private/CurveEditor/SequencerCurveEditorApp.cpp b/Engine/Source/Editor/Sequencer/Private/CurveEditor/SequencerCurveEditorApp.cpp
--- a/Engine/Source/Editor/Sequencer/Private/CurveEditor/SequencerCurveEditorApp.cpp
+++ b/Engine/Source/Editor/Sequencer/Private/CurveEditor/SequencerCurveEditorApp.cpp
@@ -114,6 +114,8 @@ FSequencerCurveEditorApp::FSequencerCurveEditorApp(const TSharedRef<FSequencer>&
 				// Sorting algorithms require consistent < implementation so fallback to pointer comparision.
 				: ViewModelA.Get() < ViewModelB.Get();
 		}));
+	
+	SelectionSyncLogic.UpdateCurveEditor();
 }
 
 FSequencerCurveEditorApp* FSequencerCurveEditorApp::Get(const ISequencer& InSequencer)
diff --git a/Engine/Source/Editor/Sequencer/Private/Filters/Linking/FilterAreaManager.cpp b/Engine/Source/Editor/Sequencer/Private/Filters/Linking/FilterAreaManager.cpp
--- a/Engine/Source/Editor/Sequencer/Private/Filters/Linking/FilterAreaManager.cpp
+++ b/Engine/Source/Editor/Sequencer/Private/Filters/Linking/FilterAreaManager.cpp
@@ -16,5 +16,7 @@ FFilterAreaManager::FFilterAreaManager(
 	, CommandBinder(InFilterAreaConfigId, InWeakSequencer, FilterModel)
 	, LinkedToUnlinkedCopier(FilterModel, InFilterAreaConfigId, InWeakSequencer)
 	, FilterModeConfigManager(InFilterAreaConfigId, InWeakSequencer, FilterModel)
-{}
+{
+	FilterModel->GetInstancedFilterBar()->RequestFilterUpdate();
+}
 }
